Adding an LRU File Cache that preallocates disk space and attempts to store files in contiguous regions, falling back to fragmenting when there is not enough space remaining.

diff --git a/Engine/Source/Runtime/VirtualFileCache/Private/VirtualFileCache.cpp b/Engine/Source/Runtime/VirtualFileCache/Private/VirtualFileCache.cpp
new file mode 100644
--- /dev/null
+++ b/Engine/Source/Runtime/VirtualFileCache/Private/VirtualFileCache.cpp
@@ -0,0 +1,98 @@
+// Copyright Epic Games, Inc. All Rights Reserved.
+
+#include "VirtualFileCache.h"
+#include "VirtualFileCacheInternal.h"
+
+#include "CoreMinimal.h"
+#include "HAL/FileManager.h"
+#include "Logging/LogMacros.h"
+#include "Misc/Paths.h"
+#include "Misc/ScopeExit.h"
+#include "Async/MappedFileHandle.h"
+
+DECLARE_STATS_GROUP(TEXT("VFC"), STATGROUP_VFC, STATCAT_Advanced);
+
+DECLARE_DWORD_COUNTER_STAT(TEXT("Files Added"), STAT_FilesAdded, STATGROUP_VFC);
+DECLARE_DWORD_COUNTER_STAT(TEXT("Bytes Added"), STAT_BytesAdded, STATGROUP_VFC);
+DECLARE_DWORD_COUNTER_STAT(TEXT("Files Removed"), STAT_FilesRemoved, STATGROUP_VFC);
+DECLARE_DWORD_COUNTER_STAT(TEXT("Bytes Removed"), STAT_BytesRemoved, STATGROUP_VFC);
+DECLARE_DWORD_COUNTER_STAT(TEXT("Files Evicted"), STAT_FilesEvicted, STATGROUP_VFC);
+DECLARE_DWORD_COUNTER_STAT(TEXT("Bytes Evicted"), STAT_BytesEvicted, STATGROUP_VFC);
+
+DECLARE_LOG_CATEGORY_EXTERN(LogVFC, Log, All);
+DEFINE_LOG_CATEGORY(LogVFC);
+
+static const TCHAR VFC_CACHE_FILE_BASE_NAME[] = TEXT("vfc_");
+static const TCHAR VFC_CACHE_FILE_EXTENSION[] = TEXT("data");
+static const TCHAR VFC_META_FILE_NAME[] = TEXT("vfc.meta");
+
+TSharedRef<IVirtualFileCache> IVirtualFileCache::CreateVirtualFileCache()
+{
+	static TWeakPtr<FVirtualFileCache> GVFC;
+	TSharedPtr<FVirtualFileCache> SharedVFC;
+	if (!GVFC.IsValid())
+	{
+		SharedVFC = MakeShared<FVirtualFileCache>();
+		GVFC = SharedVFC;
+	}
+	return GVFC.IsValid() ? GVFC.Pin().ToSharedRef() : SharedVFC.ToSharedRef();
+}
+
+void FVirtualFileCache::Shutdown()
+{
+	Thread.Shutdown();
+}
+
+void FVirtualFileCache::Initialize(const FVirtualFileCacheSettings& InSettings)
+{
+	Settings = InSettings;
+	IFileManager& FileManager = IFileManager::Get();
+	BasePath = !Settings.OverrideDefaultDirectory.IsEmpty() ? Settings.OverrideDefaultDirectory : GetVFCDirectory();
+	if (!FileManager.DirectoryExists(*BasePath))
+	{
+		FileManager.MakeDirectory(*BasePath, true);
+	}
+
+	FFileTableWriter FileTable = Thread.ModifyFileTable();
+	FileTable->Initialize(Settings);
+}
+
+FIoStatus FVirtualFileCache::WriteData(VFCKey Id, const uint8* Data, uint64 DataSize)
+{
+	Thread.RequestWrite(Id, MakeArrayView(Data, DataSize));
+	return FIoStatus(EIoErrorCode::Ok);
+}
+
+TFuture<TArray<uint8>> FVirtualFileCache::ReadData(VFCKey Id, int64 ReadOffset, int64 ReadSizeOrZero)
+{
+	return Thread.RequestRead(Id, ReadOffset, ReadSizeOrZero);
+}
+
+bool FVirtualFileCache::DoesChunkExist(const VFCKey& Id) const
+{
+	FFileTableReader FileTable = Thread.ReadFileTable();
+	return FileTable->DoesChunkExist(Id);
+}
+
+TIoStatusOr<uint64> FVirtualFileCache::GetSizeForChunk(const VFCKey& Id) const
+{
+	FFileTableReader FileTable = Thread.ReadFileTable();
+	return FileTable->GetSizeForChunk(Id);
+}
+
+void FVirtualFileCache::EraseData(VFCKey Id)
+{
+	Thread.RequestErase(Id);
+}
+
+double FVirtualFileCache::CurrentFragmentation() const
+{
+	FFileTableReader FileTable = Thread.ReadFileTable();
+	return FileTable->CurrentFragmentation();
+}
+
+void FVirtualFileCache::Defragment()
+{
+	FFileTableWriter FileTable = Thread.ModifyFileTable();
+	return FileTable->Defragment();
+}
diff --git a/Engine/Source/Runtime/VirtualFileCache/Private/VirtualFileCacheInternal.cpp b/Engine/Source/Runtime/VirtualFileCache/Private/VirtualFileCacheInternal.cpp
new file mode 100644
--- /dev/null
+++ b/Engine/Source/Runtime/VirtualFileCache/Private/VirtualFileCacheInternal.cpp
@@ -0,0 +1,1363 @@
+// Copyright Epic Games, Inc. All Rights Reserved.
+
+#include "VirtualFileCacheInternal.h"
+
+#include "CoreMinimal.h"
+#include "HAL/FileManager.h"
+#include "Logging/LogMacros.h"
+#include "Misc/Paths.h"
+#include "Misc/ScopeExit.h"
+#include "Async/MappedFileHandle.h"
+
+static int64 CurrentTimestamp()
+{
+	return FDateTime::UtcNow().ToUnixTimestamp();
+}
+
+FVirtualFileCacheThread::FVirtualFileCacheThread(FVirtualFileCache* InParent)
+	: Parent(InParent)
+	, Event(FGenericPlatformProcess::GetSynchEventFromPool(false))
+	, bStopRequested(false)
+	, FileTableStorage(InParent)
+{
+	Thread = FRunnableThread::Create(this, TEXT("FVirtualFileCacheThread"));
+}
+
+FVirtualFileCacheThread::~FVirtualFileCacheThread()
+{
+	delete Thread;
+}
+
+bool FVirtualFileCacheThread::Init()
+{
+	return Parent != nullptr;
+}
+
+void FVirtualFileCacheThread::DoOneOp(FRWOp* Op)
+{
+	switch (Op->Op)
+	{
+	case ERWOp::Read:
+	{
+		FFileTableMutator FileTable = MutateFileTable();
+		auto Result = FileTable->ReadData(Op->Target, Op->ReadOffset, Op->ReadSize);
+		if (Result.IsOk())
+		{
+			Op->ReadResult->SetValue(Result.ConsumeValueOrDie());
+		}
+		else
+		{
+			TArray<uint8> Empty;
+			Op->ReadResult->SetValue(Empty);
+		}
+	} break;
+
+	case ERWOp::Write:
+	{
+		FFileTableWriter FileTable = ModifyFileTable();
+		FIoStatus Status = FileTable->WriteData(Op->Target, Op->DataToWrite.GetData(), Op->DataToWrite.Num());
+	} break;
+
+	case ERWOp::Erase:
+	{
+		FFileTableWriter FileTable = ModifyFileTable();
+		FileTable->EraseData(Op->Target);
+	} break;
+	}
+}
+
+uint32 FVirtualFileCacheThread::Run()
+{
+	while (!bStopRequested)
+	{
+		Event->Wait(100);
+		TSharedPtr<FRWOp> Op = GetNextOp();
+		if (bStopRequested)
+		{
+			// Finish the queue or break here?
+			break;
+		}
+		if (!Op.IsValid())
+		{
+			continue;
+		}
+		DoOneOp(Op.Get());
+	}
+
+	return 0;
+}
+
+void FVirtualFileCacheThread::Stop()
+{
+	bStopRequested = true;
+	Event->Trigger();
+}
+
+void FVirtualFileCacheThread::Exit()
+{
+}
+
+TSharedPtr<FRWOp> FVirtualFileCacheThread::GetNextOp()
+{
+	TSharedPtr<FRWOp> Result;
+	FRWScopeLock ScopeLock(OperationQueueLock, SLT_Write);
+	if (!OperationQueue.IsEmpty())
+	{
+		Result = OperationQueue[0];
+		OperationQueue.RemoveAt(0);
+	}
+	return Result;
+}
+
+void FVirtualFileCacheThread::EnqueueOrRunOp(TSharedPtr<FRWOp> Op)
+{
+	if (Parent->Settings.bMultiThreaded)
+	{
+		{
+			FRWScopeLock ScopeLock(OperationQueueLock, SLT_Write);
+			OperationQueue.Add(Op);
+		}
+		Event->Trigger();
+	}
+	else
+	{
+		DoOneOp(Op.Get());
+	}
+}
+
+TFuture<TArray<uint8>> FVirtualFileCacheThread::RequestRead(VFCKey Target, int64 ReadOffset, int64 ReadSizeOrZero)
+{
+	TSharedPtr<FRWOp> Op = MakeShared<FRWOp>();
+	Op->Op = ERWOp::Read;
+	Op->Target = Target;
+	Op->ReadOffset = ReadOffset;
+	Op->ReadSize = ReadSizeOrZero;
+	Op->ReadResult.Emplace();
+	auto Result = Op->ReadResult->GetFuture();
+
+	EnqueueOrRunOp(Op);
+	return Result;
+}
+
+void FVirtualFileCacheThread::RequestWrite(VFCKey Target, TArrayView<const uint8> Data)
+{
+	TSharedPtr<FRWOp> Op = MakeShared<FRWOp>();
+	Op->Op = ERWOp::Write;
+	Op->Target = Target;
+	Op->DataToWrite = MoveTemp(Data);
+
+	EnqueueOrRunOp(Op);
+	Event->Trigger();
+}
+
+void FVirtualFileCacheThread::RequestErase(VFCKey Target)
+{
+	TSharedPtr<FRWOp> Op = MakeShared<FRWOp>();
+	Op->Op = ERWOp::Erase;
+	Op->Target = Target;
+
+	EnqueueOrRunOp(Op);
+	Event->Trigger();
+}
+
+void FVirtualFileCacheThread::Shutdown()
+{
+	bool bValidRanges = false;
+	{
+		FFileTableMutator FileTable = MutateFileTable();
+		bValidRanges = FileTable->ValidateRanges();
+	}
+
+	if (bValidRanges)
+	{
+		FFileTableWriter FileTable = ModifyFileTable();
+
+		double Fragmentation = FileTable->CurrentFragmentation();
+		UE_LOG(LogVFC, Log, TEXT("Current fragmentation is %0.4f (1.0 == no fragmentation)"), Fragmentation);
+		if (Fragmentation > VFC_ALLOWED_FRAGMENTATION)
+		{
+			FileTable->Defragment();
+		}
+		else
+		{
+			FileTable->CoalesceRanges();
+		}
+
+		FileTable->WriteTableFile();
+	}
+	else
+	{
+		UE_LOG(LogVFC, Verbose, TEXT("Deleting table file due to an invalid range\n"));
+		EraseTableFile();
+	}
+}
+
+FFileTableReader FVirtualFileCacheThread::ReadFileTable() const
+{
+	return FFileTableReader(FileTableStorage, FileTableLock);
+}
+
+FFileTableMutator FVirtualFileCacheThread::MutateFileTable()
+{
+	return FFileTableMutator(FileTableStorage, FileTableLock);
+}
+
+FFileTableWriter FVirtualFileCacheThread::ModifyFileTable()
+{
+	return FFileTableWriter(FileTableStorage, FileTableLock);
+}
+
+static void InsertOrdered(TArray<FBlockRange>& BlockRanges, FBlockRange NewRange)
+{
+	check(NewRange.NumBlocks > 0);
+	int32 Index = 0;
+	for (; Index < BlockRanges.Num(); ++Index)
+	{
+		if (BlockRanges[Index].NumBlocks < NewRange.NumBlocks)
+		{
+			break;
+		}
+	}
+	BlockRanges.Insert(NewRange, Index);
+}
+
+static bool EraseOrdered(TArray<FBlockRange>& BlockRanges, FBlockRange Range)
+{
+	int32 Count = BlockRanges.Num();
+	for (int32 Index = 0; Index < Count; ++Index)
+	{
+		if (Range.StartIndex == BlockRanges[Index].StartIndex)
+		{
+			BlockRanges.RemoveAt(Index);
+			return true;
+		}
+	}
+
+	return false;
+}
+
+void FFileTable::FreeBlock(FRangeId RangeId)
+{
+	FBlockFile* BlockFile = GetFileForRange(RangeId);
+	if (BlockFile)
+	{
+		bool Erased = EraseOrdered(BlockFile->UsedRanges, RangeId.Range);
+		if (ensure(Erased))// TODO: Need to revalidate the entire file if this is not found.
+		{
+			InsertOrdered(BlockFile->FreeRanges, RangeId.Range);
+			UsedSize -= RangeId.Range.NumBlocks * BlockFile->BlockSize;
+		}
+	}
+}
+
+FIoStatus FFileTable::WriteData(VFCKey Id, const uint8* Data, uint64 DataSize)
+{
+	if (DoesChunkExist(Id))
+	{
+		EraseData(Id);
+	}
+
+	const int64 Size = int64(DataSize);
+	bool bSpaceIsAvailable = EnsureSizeFor(Size);
+	if (!bSpaceIsAvailable)
+	{
+		UE_LOG(LogVFC, Error, TEXT("Unable to fit data into cache"));
+		return FIoStatus(EIoErrorCode::NotFound);
+	}
+
+	FIoStatus Result = FIoStatus(EIoErrorCode::Ok);
+	TArray<FRangeId> Blocks = AllocateBlocksForSize(Size);
+	int64 TotalAllocatedSize = 0;
+
+	if (Blocks.Num() == 0)
+	{
+		Result = FIoStatus(EIoErrorCode::NotFound);
+	}
+	else
+	{
+		FDataReference DataRef;
+		int64 DataWritten = 0;
+
+		for (const FRangeId& RangeId : Blocks)
+		{
+			FBlockFile* BlockFile = GetFileForRange(RangeId);
+			IFileHandle* Handle = BlockFile ? OpenBlockFileForWrite(*BlockFile) : nullptr;
+			if (Handle == nullptr)
+			{
+				Result = EIoErrorCode::NotFound;
+				break;
+			}
+
+			if (RangeIsValid(RangeId))
+			{
+				FBlockRange Range = RangeId.Range;
+				int64 Offset = Range.StartIndex * BlockFile->BlockSize;
+				int64 ContiguousSize = BlockFile->BlockSize * Range.NumBlocks;
+				int64 WriteSize = FMath::Min(Size - DataWritten, ContiguousSize);
+
+				Handle->Seek(Offset);
+				Handle->Write(Data + DataWritten, WriteSize);
+
+				DataWritten += WriteSize;
+				TotalAllocatedSize += ContiguousSize;
+
+				DataRef.Ranges.Add(RangeId);
+			}
+
+			//Handle->Flush();
+		}
+
+		UE_LOG(LogVFC, Verbose, TEXT("Wrote %llu bytes\n"), DataSize);
+
+		DataRef.TotalSize = DataSize;
+		DataRef.Touch();
+		FileMap.Add(Id, DataRef);
+	}
+
+	if (Result.IsOk())
+	{
+		UsedSize += TotalAllocatedSize;
+		WriteTableFile();
+
+		INC_DWORD_STAT(STAT_FilesAdded);
+		INC_DWORD_STAT_BY(STAT_BytesAdded, DataSize);
+	}
+	else
+	{
+		// Free the blocks on error
+		for (FRangeId& RangeId : Blocks)
+		{
+			FreeBlock(RangeId);
+		}
+	}
+
+	return Result;
+}
+
+bool FFileTable::ReadRange(FRangeId RangeId, void* Dest, int64 ReadSize)
+{
+	FBlockFile* BlockFile = GetFileForRange(RangeId);
+	if (!BlockFile)
+	{
+		return false;
+	}
+
+	bool bSuccess = false;
+	FString BlockFilePath = GetBlockFilename(*BlockFile);
+
+	FRWScopeLock ScopeLock(*BlockFile->FileHandleLock.Get(), SLT_Write);
+	BlockFile->WriteHandle = nullptr;
+	TUniquePtr<IFileHandle> FileHandle = TUniquePtr<IFileHandle>(IPlatformFile::GetPlatformPhysical().OpenRead(*BlockFilePath, false));
+	if (FileHandle.IsValid())
+	{
+		int64 Offset = RangeId.Range.StartIndex * BlockFile->BlockSize;
+		int64 RangeSize = RangeId.Range.NumBlocks * BlockFile->BlockSize;
+		int64 Size = FMath::Min(RangeSize, ReadSize);
+		if (FileHandle->Seek(Offset))
+		{
+			if (Dest)
+			{
+				if (FileHandle->Read((uint8*)Dest, Size))
+				{
+					bSuccess = true;
+				}
+			}
+			else
+			{
+				bSuccess = FileHandle->Size() >= (Offset + Size);
+			}
+		}
+	}
+	return bSuccess;
+}
+
+TIoStatusOr<TArray<uint8>> FFileTable::ReadData(VFCKey Id, int64 ReadOffset, int64 ReadSizeOrZero)
+{
+	// Read lock must be maintained while we have a pointer to a value in the FileMap
+	FDataReference* DataRef = FileMap.Find(Id);
+	if (DataRef == nullptr)
+	{
+		return FIoStatus(EIoErrorCode::NotFound);
+	}
+	FIoStatus Status = EIoErrorCode::Ok;
+
+	int64 SeekRemaining = ReadOffset;
+	int64 SizeRemaining = ReadSizeOrZero > 0 ? ReadSizeOrZero : (DataRef->TotalSize - ReadOffset);
+	check(SizeRemaining + SeekRemaining <= DataRef->TotalSize);
+	int64 WriteOffset = 0;
+	int64 NumBlocksRead = 0;
+	DataRef->Touch();
+
+	TArray<uint8> Result;
+	Result.SetNumUninitialized(SizeRemaining);
+
+	for (const FRangeId& RangeId : DataRef->Ranges)
+	{
+		if (RangeId.Range.NumBlocks == 0)
+		{
+			UE_LOG(LogVFC, Warning, TEXT("Invalid block ID (run length of zero)"));
+			Status = EIoErrorCode::NotFound;
+			break;
+		}
+
+		if (SizeRemaining == 0)
+		{
+			break;
+		}
+
+		FBlockFile* BlockFile = GetFileForRange(RangeId);
+		if (BlockFile == nullptr)
+		{
+			Status = EIoErrorCode::NotFound;
+			break;
+		}
+
+		FBlockRange Range = RangeId.Range;
+		int64 ContiguousSize = Range.NumBlocks * BlockFile->BlockSize;
+		int64 Offset = Range.StartIndex * BlockFile->BlockSize;
+
+		if (SeekRemaining > 0)
+		{
+			int64 SeekSize = FMath::Min(SeekRemaining, ContiguousSize);
+			SeekRemaining -= SeekSize;
+			ContiguousSize -= SeekSize;
+			Offset += SeekSize;
+		}
+
+		if (ContiguousSize > 0)
+		{
+			check(SeekRemaining == 0);
+			int64 ReadSize = FMath::Min(SizeRemaining, ContiguousSize);
+
+			FMappedRange MappedRange = MapFileRange(RangeId);
+			if (MappedRange.IsValid())
+			{
+				check(ReadSize <= MappedRange->GetMappedSize());
+				FMemory::Memcpy(&Result[WriteOffset], MappedRange->GetMappedPtr(), ReadSize);
+			}
+			else if (!ReadRange(RangeId, &Result[WriteOffset], ReadSize))
+			{
+				Status = EIoErrorCode::ReadError;
+				break;
+			}
+
+			SizeRemaining -= ReadSize;
+			WriteOffset += ReadSize;
+		}
+
+		NumBlocksRead++;
+	}
+
+	if (Status != EIoErrorCode::Ok)
+	{
+		return Status;
+	}
+	else
+	{
+		check(SizeRemaining == 0);
+		return MoveTemp(Result);
+	}
+}
+
+const FDataReference* FFileTable::FindDataReference(const VFCKey& Id) const
+{
+	return FileMap.Find(Id);
+}
+
+TIoStatusOr<uint64> FFileTable::GetSizeForChunk(const VFCKey& Id) const
+{
+	const FDataReference* DataRef = FindDataReference(Id);
+	if (DataRef)
+	{
+		return DataRef->TotalSize;
+	}
+	return FIoStatus(EIoErrorCode::NotFound);
+}
+
+bool FFileTable::EraseData(VFCKey Id)
+{
+	FDataReference* DataRef = FileMap.Find(Id);
+	if (DataRef)
+	{
+		INC_DWORD_STAT(STAT_FilesRemoved);
+		INC_DWORD_STAT_BY(STAT_BytesRemoved, DataRef->TotalSize);
+
+		for (FRangeId& RangeId : DataRef->Ranges)
+		{
+			FreeBlock(RangeId);
+		}
+
+		check(UsedSize >= 0 && UsedSize <= TotalSize);
+		FileMap.Remove(Id);
+		WriteTableFile();
+
+		return true;
+	}
+	return false;
+}
+
+int64 FFileTable::EvictOne()
+{
+	int64 RemovedSize = 0;
+	VFCKey LeastRecentlyUsed = VFCKey();
+	int64 LeastRecentTime = INT64_MAX;
+	for (auto& [MapKey, MapValue] : FileMap)
+	{
+		if (MapValue.LastReferencedUnixTime >= 0 && MapValue.LastReferencedUnixTime < LeastRecentTime)
+		{
+			UE_LOG(LogVFC, Verbose, TEXT("Evicting Data %u with size of %u\n"), GetTypeHash(MapKey), MapValue.TotalSize);
+			LeastRecentlyUsed = MapKey;
+			LeastRecentTime = MapValue.LastReferencedUnixTime;
+		}
+	}
+
+	FDataReference* ToRemove = FileMap.Find(LeastRecentlyUsed);
+	if (ToRemove)
+	{
+		RemovedSize = AllocationSize(ToRemove);
+		Parent->OnDataEvicted.Broadcast(LeastRecentlyUsed);
+		EraseData(LeastRecentlyUsed);
+
+		INC_DWORD_STAT(STAT_FilesEvicted);
+		INC_DWORD_STAT_BY(STAT_BytesEvicted, RemovedSize);
+	}
+
+	return RemovedSize;
+}
+
+bool FFileTable::EvictAmount(int64 NumBytesToEvict)
+{
+	int64 Evicted = 0;
+
+	while (Evicted < NumBytesToEvict)
+	{
+		int64 EvictedSize = EvictOne();
+		if (EvictedSize == 0)
+		{
+			break;
+		}
+		Evicted += EvictedSize;
+	}
+
+	return Evicted >= NumBytesToEvict;
+}
+
+bool FFileTable::EnsureSizeFor(int64 RequiredBytes)
+{
+	bool bSizeIsAvailable = true;
+	int64 AvailableSize = TotalSize - UsedSize;
+
+	if (AvailableSize < RequiredBytes)
+	{
+		int64 RequiredEvictionSize = RequiredBytes - AvailableSize;
+		bSizeIsAvailable = EvictAmount(RequiredEvictionSize);
+	}
+
+	return bSizeIsAvailable;
+}
+
+TIoStatusOr<FRangeId> FFileTable::AllocateSingleRange(FBlockFile& File, int64 MaximumSize)
+{
+	check(MaximumSize > 0);
+	if (File.FreeRanges.IsEmpty() || File.bWriteLocked)
+	{
+		return FIoStatus(EIoErrorCode::NotFound);
+	}
+
+	int32 BestIndex = 0; // Grabbing first element for now.
+
+	FBlockRange Source = File.FreeRanges[BestIndex];
+	int64 MaximumBlocks = Align(MaximumSize, File.BlockSize) / File.BlockSize;
+	int64 Size = Source.NumBlocks * File.BlockSize;
+
+	if (Source.NumBlocks > MaximumBlocks)
+	{
+		// Split the block up since it is larger than what we need
+		FBlockRange Allocated;
+		Allocated.StartIndex = Source.StartIndex;
+		Allocated.NumBlocks = MaximumBlocks;
+
+		FBlockRange Remaining;
+		Remaining.StartIndex = Allocated.StartIndex + Allocated.NumBlocks;
+		Remaining.NumBlocks = Source.NumBlocks - MaximumBlocks;
+
+		File.FreeRanges.RemoveAt(BestIndex);
+		InsertOrdered(File.UsedRanges, Allocated); // Mark range as used
+		check(Remaining.NumBlocks > 0);
+		InsertOrdered(File.FreeRanges, Remaining); // Mark remaining as available
+
+		Source = Allocated;
+	}
+	else
+	{
+		// We need the whole block so move it into the used list
+		File.FreeRanges.RemoveAt(BestIndex);
+		InsertOrdered(File.UsedRanges, Source);
+	}
+
+	FRangeId Result = {};
+	Result.FileId = File.FileId;
+	Result.Range = Source;
+	return Result;
+}
+
+TArray<FRangeId> FFileTable::AllocateBlocksForSize(uint64 Size)
+{
+	TArray<FRangeId> Result;
+	int64 SizeRemaining = Size;
+	for (FBlockFile& File : BlockFiles)
+	{
+		if (File.bWriteLocked)
+		{
+			continue;
+		}
+
+		if (File.FreeRanges.IsEmpty())
+		{
+			continue;
+		}
+
+		int32 RequiredBlocks = Align(Size, File.BlockSize) / File.BlockSize;
+		// FreeRanges is sorted by size, so if the first one cannot fit the data, then no single block in this file will be able to
+		if (File.FreeRanges[0].NumBlocks < RequiredBlocks)
+		{
+			continue;
+		}
+
+		int32 BestIndex = 0;
+		for (int32 Index = 1; Index < File.FreeRanges.Num(); Index++)
+		{
+			if (File.FreeRanges[Index].NumBlocks < RequiredBlocks)
+			{
+				break;
+			}
+			BestIndex = Index;
+		}
+
+		// Split Block
+		FBlockRange Source = File.FreeRanges[BestIndex];
+
+		FBlockRange Allocated;
+		Allocated.StartIndex = Source.StartIndex;
+		Allocated.NumBlocks = RequiredBlocks;
+
+		FBlockRange Remaining;
+		Remaining.StartIndex = Allocated.StartIndex + Allocated.NumBlocks;
+		Remaining.NumBlocks = Source.NumBlocks - RequiredBlocks;
+
+		File.FreeRanges.RemoveAt(BestIndex);
+		InsertOrdered(File.UsedRanges, Allocated); // Mark range as used
+		if (Remaining.NumBlocks > 0)
+		{
+			InsertOrdered(File.FreeRanges, Remaining); // Mark remaining as available
+		}
+
+		FRangeId RangeId = {};
+		RangeId.FileId = File.FileId;
+		RangeId.Range = Allocated;
+		Result.Add(RangeId);
+		SizeRemaining = 0;
+		break;
+	}
+
+	if (SizeRemaining > 0)
+	{
+		for (FBlockFile& File : BlockFiles)
+		{
+			if (File.bWriteLocked)
+			{
+				continue;
+			}
+
+			while (SizeRemaining > 0)
+			{
+				TIoStatusOr<FRangeId> ResultOrRange = AllocateSingleRange(File, SizeRemaining);
+				if (ResultOrRange.IsOk())
+				{
+					FRangeId RangeId = ResultOrRange.ValueOrDie();
+					SizeRemaining -= RangeId.Range.NumBlocks * File.BlockSize;
+					Result.Add(RangeId);
+				}
+				else
+				{
+					// This file doesn't have any room left, check the remaining files.
+					break;
+				}
+			}
+		}
+		check(SizeRemaining <= 0);
+	}
+
+	if (SizeRemaining > 0)
+	{
+		UE_LOG(LogVFC, Verbose, TEXT("Not enough space to satisfy request. Needed %llu bytes, only had %llu available\n"), Size, Size - SizeRemaining);
+		// Error, not enough space in the cache file. Allocate a new one or evict
+		for (FRangeId Id : Result)
+		{
+			FreeBlock(Id);
+		}
+		Result.Empty();
+	}
+	return Result;
+}
+
+int64 FFileTable::AllocationSize(FDataReference* DataRef)
+{
+	int64 AllocatedSize = 0;
+	for (const FRangeId& Range : DataRef->Ranges)
+	{
+		FBlockFile* BlockFile = GetFileForRange(Range);
+		if (BlockFile)
+		{
+			AllocatedSize += BlockFile->BlockSize * Range.Range.NumBlocks;
+		}
+	}
+	return AllocatedSize;
+}
+
+void FFileTable::Defragment()
+{
+	FFileMap BackupMap = FileMap;
+	TArray<int32> CurrentFiles;
+	TArray<int32> AddedFiles;
+	bool bSuccess = true;
+
+	CurrentFiles.Empty(BlockFiles.Num());
+	for (FBlockFile& BlockFile : BlockFiles)
+	{
+		CurrentFiles.Add(BlockFile.FileId);
+	}
+
+	// Create a new set of block files to defragment into
+	for (int32 BlockFileId : CurrentFiles)
+	{
+		FBlockFile* BlockFile = GetFileForId(BlockFileId);
+		if (ensure(BlockFile))
+		{
+			BlockFile->bWriteLocked = true;
+			int32 NewFileId = CreateBlockFile(BlockFile->TotalSize(), BlockFile->BlockSize);
+			if (NewFileId > 0)
+			{
+				AddedFiles.Add(NewFileId);
+			}
+			else
+			{
+				bSuccess = false;
+				break;
+			}
+		}
+	}
+	CalculateSizes();
+
+	int32 TESTCounter = 0;
+
+	if (bSuccess)
+	{
+		TArray<VFCKey> DataToMove;
+		for (auto& [Key, Value] : FileMap)
+		{
+			DataToMove.Add(Key);
+		}
+
+		for (const VFCKey& Id : DataToMove)
+		{
+			TIoStatusOr<TArray<uint8>> ReadResult = ReadData(Id);
+
+			// Manually remove from the map without calling EraseData since the entire file that contains
+			// this data will be deleted after this finishes. The key can't be in the map or it won't write.
+			FileMap.Remove(Id);
+
+			if (ensure(ReadResult.IsOk()))
+			{
+				TArray<uint8> Data = ReadResult.ConsumeValueOrDie();
+				FIoStatus WriteResult = WriteData(Id, Data.GetData(), Data.Num());
+
+				if (!WriteResult.IsOk() || TESTCounter++ == 5)
+				{
+					bSuccess = false;
+					break;
+				}
+			}
+			else
+			{
+				bSuccess = false;
+				break;
+			}
+		}
+	}
+
+	if (bSuccess)
+	{
+		for (int32 ToRemove : CurrentFiles)
+		{
+			DeleteBlockFile(ToRemove);
+		}
+		WriteTableFile();
+	}
+	else
+	{
+		for (int32 ToRemove : AddedFiles)
+		{
+			DeleteBlockFile(ToRemove);
+		}
+		for (FBlockFile& BlockFile : BlockFiles)
+		{
+			BlockFile.bWriteLocked = false;
+		}
+		FileMap = MoveTemp(BackupMap);
+	}
+
+	CalculateSizes();
+	ValidateRanges();
+}
+
+int32 FFileTable::CreateBlockFile(int64 FileSize, int32 BlockSize)
+{
+	int32 FileId = ++LastBlockFileId;
+	check(GetFileForId(FileId) == nullptr);
+
+	FBlockFile BlockFile = {};
+	BlockFile.FileId = FileId;
+	BlockFile.BlockSize = BlockSize;
+
+	FString ChunkFilePath = GetBlockFilename(BlockFile);
+	IPlatformFile& PlatformFile = IPlatformFile::GetPlatformPhysical();
+	TUniquePtr<IFileHandle> NewFile = TUniquePtr<IFileHandle>(PlatformFile.OpenWrite(*ChunkFilePath));
+
+	if (NewFile.IsValid())
+	{
+		uint64 NumBlocks = FileSize / BlockSize;
+		uint64 ActualSize = NumBlocks * BlockSize;
+		check(NumBlocks <= TNumericLimits<decltype(FBlockRange::NumBlocks)>::Max());
+
+		FBlockRange FullRange = { 0, (int32)NumBlocks };
+		BlockFile.FreeRanges.Add(FullRange);
+		BlockFile.NumBlocks = (int32)NumBlocks;
+
+		// Pre-allocate disk space
+		NewFile->Seek(ActualSize - 1);
+		char NullData = '!';
+		NewFile->Write((uint8*)&NullData, 1);
+
+		BlockFiles.Add(MoveTemp(BlockFile));
+		return FileId;
+	}
+
+	UE_LOG(LogVFC, Error, TEXT("Unable to create backing file"));
+	return 0;
+}
+
+bool FFileTable::DeleteBlockFile(int32 BlockFileId)
+{
+	bool bDeleted = false;
+	if (FBlockFile* BlockFile = GetFileForId(BlockFileId))
+	{
+		BlockFile->WriteHandle.Reset();
+		FString Path = GetBlockFilename(*BlockFile);
+		bDeleted = IPlatformFile::GetPlatformPhysical().DeleteFile(*Path);
+		BlockFiles.RemoveAll([BlockFileId](FBlockFile& BF) { return BF.FileId == BlockFileId; });
+	}
+	return bDeleted;
+}
+
+void FFileTable::WriteTableFile()
+{
+	FString TableFileName = Parent->GetTableFilename();
+	TUniquePtr<FArchive> ArchivePtr = TUniquePtr<FArchive>(IFileManager::Get().CreateFileWriter(*TableFileName));
+	if (ArchivePtr)
+	{
+		FArchive& Ar = *ArchivePtr;
+		Ar << *this;
+	}
+}
+
+bool FFileTable::ReadTableFile()
+{
+	bool bSuccess = false;
+	FString TableFilename = Parent->GetTableFilename();
+	TUniquePtr<FArchive> ArchivePtr = TUniquePtr<FArchive>(IFileManager::Get().CreateFileReader(*TableFilename));
+	if (ArchivePtr.IsValid())
+	{
+		*ArchivePtr << *this;
+		if (ArchivePtr->IsError() || !ValidateRanges())
+		{
+			Empty();
+			ArchivePtr.Reset();
+			IFileManager::Get().Delete(*TableFilename);
+			UE_LOG(LogVFC, Log, TEXT("Invalid table file"));
+		}
+		else
+		{
+			bSuccess = true;
+		}
+	}
+	else
+	{
+		UE_LOG(LogVFC, Log, TEXT("No table file"));
+	}
+
+	return bSuccess;
+}
+
+void DeleteUnexpectedCacheFiles(const FString& BasePath, TSet<FString>& ExpectedFiles)
+{
+	struct FDeleteUnexpectedFilesVisitor : IPlatformFile::FDirectoryVisitor
+	{
+		FDeleteUnexpectedFilesVisitor(TSet<FString>& InExpectedFiles)
+			: ExpectedFiles(InExpectedFiles)
+		{}
+		TSet<FString>& ExpectedFiles;
+
+		virtual bool Visit(const TCHAR* FilenameOrDirectory, bool bIsDirectory) override
+		{
+			if (!bIsDirectory)
+			{
+				FString FullPath = FilenameOrDirectory;
+				FString Extension = FPaths::GetExtension(FullPath);
+				if (Extension == VFC_CACHE_FILE_EXTENSION)
+				{
+					FString Filename = FPaths::GetCleanFilename(FullPath);
+					if (!ExpectedFiles.Contains(Filename))
+					{
+						IPlatformFile& PlatformFile = IPlatformFile::GetPlatformPhysical();
+						PlatformFile.DeleteFile(*FullPath);
+					}
+				}
+			}
+			return true;
+		}
+	};
+	FDeleteUnexpectedFilesVisitor Visitor(ExpectedFiles);
+	IPlatformFile& PlatformFile = IPlatformFile::GetPlatformPhysical();
+
+	PlatformFile.IterateDirectory(*BasePath, Visitor);
+}
+
+void FFileTable::Initialize(const FVirtualFileCacheSettings& Settings)
+{
+	ReadTableFile();
+
+	TSet<FString> ExpectedFiles;
+	TSet<int32> MissingBlockFileIds;
+	int32 NumRemovedBlockFiles = BlockFiles.RemoveAll([this, &ExpectedFiles, &MissingBlockFileIds](FBlockFile& BlockFile) {
+		IFileHandle* File = OpenBlockFileForWrite(BlockFile);
+		if (File == nullptr || File->Size() != BlockFile.TotalSize())
+		{
+			// Remove the reference to this chunk file if the file is not on disk
+			MissingBlockFileIds.Add(BlockFile.FileId);
+			return true;
+		}
+
+		FString FoundFile = GetBlockFilename(BlockFile);
+		ExpectedFiles.Add(FoundFile);
+
+		return false;
+	});
+
+	if (!MissingBlockFileIds.IsEmpty())
+	{
+		// An expected block file was removed, clear the list of files in the cache.
+		Empty();
+		for (FBlockFile& BlockFile : BlockFiles)
+		{
+			BlockFile.UsedRanges.Empty();
+
+			BlockFile.FreeRanges.Empty();
+			FBlockRange FullRange = {};
+			FullRange.StartIndex = 0;
+			FullRange.NumBlocks = BlockFile.NumBlocks;
+			BlockFile.FreeRanges.Add(FullRange);
+		}
+	}
+
+	if (NumRemovedBlockFiles > 0)
+	{
+		// We found and removed invalid data, rewrite the table with the references removed.
+		WriteTableFile();
+	}
+
+	DeleteUnexpectedCacheFiles(Parent->BasePath, ExpectedFiles);
+
+	if (BlockFiles.Num() < Settings.NumBlockFiles)
+	{
+		int64 BlockSize = FMath::Clamp(Settings.BlockSize, VFC_BLOCK_SIZE_MIN, VFC_BLOCK_SIZE_MAX);
+		int64 CacheFileSize = Settings.BlockFileSize;
+		BlockSize = FMath::Min(BlockSize, CacheFileSize);
+
+		for (int32 i = 0; i < Settings.NumBlockFiles; ++i)
+		{
+			if (CreateBlockFile(CacheFileSize, BlockSize) == 0)
+			{
+				break;
+			}
+		}
+	}
+
+	CalculateSizes();
+}
+
+
+
+void FFileTable::Empty()
+{
+	BlockFiles.Empty();
+	FileMap.Empty();
+	TotalSize = -1;
+	UsedSize = -1;
+}
+
+double FFileTable::CurrentFragmentation() const
+{
+	double RangesPerFile = 0;
+	int32 TotalFiles = 0;
+	for (auto& [Key, Value] : FileMap)
+	{
+		TotalFiles += 1;
+		RangesPerFile += Value.Ranges.Num();
+	}
+
+	double AverageRangesPerFile = RangesPerFile / (double)TotalFiles;
+	return AverageRangesPerFile;
+}
+
+bool FFileTable::ValidateRanges()
+{
+	bool bValid = true;
+
+	for (const FBlockFile& BlockFile : BlockFiles)
+	{
+		TArray<FBlockRange> AllRanges;
+		AllRanges.Append(BlockFile.UsedRanges);
+		AllRanges.Append(BlockFile.FreeRanges);
+
+		const FBlockRangeSortStartIndex SortPredicate;
+		AllRanges.Sort(SortPredicate);
+
+		int32 NextExpectedStart = 0;
+		for (FBlockRange Range : AllRanges)
+		{
+			if (Range.StartIndex != NextExpectedStart)
+			{
+				bValid = false;
+				break;
+			}
+			NextExpectedStart = Range.StartIndex + Range.NumBlocks;
+		}
+
+		if (bValid)
+		{
+			bValid = NextExpectedStart == BlockFile.NumBlocks && BlockFileExistsOnDisk(BlockFile);
+		}
+
+		if (!bValid)
+		{
+			break;
+		}
+	}
+
+	if (bValid)
+	{
+		for (auto& [Key, Value] : FileMap)
+		{
+			int64 RangeSize = 0;
+			for (const FRangeId& RangeId : Value.Ranges)
+			{
+				const FBlockFile* BlockFile = GetFileForRange(RangeId);
+				if (BlockFile)
+				{
+					FMappedRange MappedRange = MapFileRange(RangeId);
+					if (!MappedRange.IsValid())
+					{
+						if (!ReadRange(RangeId, nullptr, INT64_MAX))
+						{
+							bValid = false;
+							break;
+						}
+					}
+					RangeSize += RangeId.Range.NumBlocks * BlockFile->BlockSize;
+				}
+				else
+				{
+					bValid = false;
+					break;
+				}
+			}
+			bValid &= RangeSize >= Value.TotalSize;
+
+			if (!bValid)
+			{
+				break;
+			}
+		}
+	}
+
+	if (!bValid)
+	{
+		UE_LOG(LogVFC, Error, TEXT("Invalid ranges"));
+	}
+
+	return bValid;
+}
+
+bool FFileTable::CoalesceRanges()
+{
+	const FBlockRangeSortStartIndex IndexSort;
+	const FBlockRangeSortSize SizeSort;
+
+	int32 RemovedRanges = 0;
+	for (FBlockFile& BlockFile : BlockFiles)
+	{
+		TArray<FBlockRange> Ranges = BlockFile.FreeRanges;
+		Ranges.Sort(IndexSort);
+
+		bool bChanged = false;
+		for (int32 Index = 1; Index < Ranges.Num();)
+		{
+			int32 PrevIndex = Index - 1;
+			FBlockRange Prev = Ranges[PrevIndex];
+			FBlockRange Curr = Ranges[Index];
+
+			if (Prev.StartIndex + Prev.NumBlocks == Curr.StartIndex)
+			{
+				// Extend previous range with the current range and remove the current range
+				Ranges[PrevIndex].NumBlocks += Curr.NumBlocks;
+				Ranges.RemoveAt(Index);
+				bChanged = true;
+
+				// Don't increment Index here since we removed the element at this index and need to check the new value that exists here.
+			}
+			else
+			{
+				Index += 1;
+			}
+		}
+
+		if (bChanged)
+		{
+			RemovedRanges += BlockFile.FreeRanges.Num() - Ranges.Num();
+			BlockFile.FreeRanges = Ranges;
+			BlockFile.FreeRanges.Sort(SizeSort);
+		}
+	}
+
+	UE_LOG(LogVFC, Verbose, TEXT("Coalesce ranges reduced free range count by %d"), RemovedRanges);
+
+	check(ValidateRanges());
+	return RemovedRanges > 0;
+}
+
+void FFileTable::CalculateSizes()
+{
+	int64 Sum = 0;
+	int64 FreeSize = 0;
+
+	for (FBlockFile& BlockFile : BlockFiles)
+	{
+		Sum += BlockFile.TotalSize();
+		int64 FreeBlocks = 0;
+		int64 UsedBlocks = 0;
+		for (const FBlockRange Range : BlockFile.FreeRanges)
+		{
+			FreeBlocks += Range.NumBlocks;
+		}
+		for (const FBlockRange Range : BlockFile.UsedRanges)
+		{
+			UsedBlocks += Range.NumBlocks;
+		}
+		FreeSize += FreeBlocks * BlockFile.BlockSize;
+	}
+
+	TotalSize = Sum;
+	UsedSize = Sum - FreeSize;
+}
+
+
+FString FFileTable::GetBlockFilename(const FBlockFile& File) const
+{
+	FString Filename = VFC_CACHE_FILE_BASE_NAME + LexToString(File.FileId);
+	Filename = FPaths::SetExtension(Filename, VFC_CACHE_FILE_EXTENSION);
+	return Parent->BasePath / Filename;
+}
+
+FString FVirtualFileCacheThread::GetTableFilename() const
+{
+	return Parent->BasePath / VFC_META_FILE_NAME;
+}
+
+FBlockFile* FFileTable::GetFileForRange(FRangeId Id)
+{
+	return GetFileForId(Id.FileId);
+}
+
+FBlockFile* FFileTable::GetFileForId(int32 BlockFileId)
+{
+	for (auto& ChunkFile : BlockFiles)
+	{
+		if (ChunkFile.FileId == BlockFileId)
+		{
+			return &ChunkFile;
+		}
+	}
+
+	return nullptr;
+}
+
+bool FFileTable::BlockFileExistsOnDisk(const FBlockFile& BlockFile) const
+{
+	return BlockFile.WriteHandle.IsValid() || IPlatformFile::GetPlatformPhysical().FileExists(*GetBlockFilename(BlockFile));
+}
+
+IFileHandle* FFileTable::OpenBlockFileForWrite(FBlockFile& BlockFile)
+{
+	FString Path = GetBlockFilename(BlockFile);
+	if (!BlockFile.WriteHandle.IsValid())
+	{
+		FRWScopeLock ScopeLock(*BlockFile.FileHandleLock.Get(), SLT_Write);
+		BlockFile.MapHandle = nullptr;
+
+		// Open for read and write
+		IPlatformFile& PlatformFile = IPlatformFile::GetPlatformPhysical();
+		BlockFile.WriteHandle = TUniquePtr<IFileHandle>(PlatformFile.OpenWrite(*Path, true, true));
+	}
+	return BlockFile.WriteHandle.Get();
+}
+
+bool FFileTable::DoesChunkExist(const VFCKey& Id) const
+{
+	return FileMap.Contains(Id);
+}
+
+FMappedRange FFileTable::MapFileRange(FRangeId RangeId)
+{
+	FBlockFile* BlockFile = GetFileForRange(RangeId);
+	if (!BlockFile)
+	{
+		return nullptr;
+	}
+	IMappedFileHandle* MapHandle = MapBlockFile(*BlockFile);
+	if (!MapHandle)
+	{
+		return nullptr;
+	}
+
+	int64 Offset = RangeId.Range.StartIndex * BlockFile->BlockSize;
+	int64 Size = RangeId.Range.NumBlocks * BlockFile->BlockSize;
+	bool bPreload = false;
+	IMappedFileRegion* R = MapHandle->MapRegion(Offset, Size, bPreload);
+	auto Result = TUniquePtr<IMappedFileRegion>(R);
+	return MoveTemp(Result);
+}
+
+IMappedFileHandle* FFileTable::MapBlockFile(FBlockFile& BlockFile)
+{
+	FString Path = GetBlockFilename(BlockFile);
+	if (!BlockFile.MapHandle.IsValid())
+	{
+		FRWScopeLock ScopeLock(*BlockFile.FileHandleLock.Get(), SLT_Write);
+		BlockFile.WriteHandle = nullptr;
+
+		// Open for read and write
+		IPlatformFile& PlatformFile = IPlatformFile::GetPlatformPhysical();
+		BlockFile.MapHandle = TUniquePtr<IMappedFileHandle>(PlatformFile.OpenMapped(*Path));
+	}
+	return BlockFile.MapHandle.Get();
+}
+
+bool FFileTable::RangeIsValid(FRangeId Block)
+{
+	FBlockFile* File = GetFileForRange(Block);
+	if (File != nullptr)
+	{
+		return Block.Range.StartIndex >= 0 && Block.Range.StartIndex + Block.Range.NumBlocks <= File->NumBlocks;
+	}
+	return false;
+}
+
+void FDataReference::Touch()
+{
+	LastReferencedUnixTime = CurrentTimestamp();
+}
+
+void FVirtualFileCacheThread::EraseTableFile()
+{
+	FString TableFileName = GetTableFilename();
+	IPlatformFile::GetPlatformPhysical().DeleteFile(*TableFileName);
+}
+
+
+static FArchive& operator<<(FArchive& Ar, FBlockRange& Range)
+{
+	Ar << Range.StartIndex;
+	Ar << Range.NumBlocks;
+	return Ar;
+}
+
+static FArchive& operator<<(FArchive& Ar, FRangeId& RangeId)
+{
+	Ar << RangeId.FileId;
+	Ar << RangeId.Range;
+	return Ar;
+}
+
+static FArchive& operator<<(FArchive& Ar, FDataReference& DataRef)
+{
+	Ar << DataRef.Ranges;
+	Ar << DataRef.LastReferencedUnixTime;
+	Ar << DataRef.TotalSize;
+	return Ar;
+}
+
+static FArchive& operator<<(FArchive& Ar, EVFCFileVersion& FileVersion)
+{
+	static_assert(sizeof(FileVersion) == sizeof(int32), "");
+
+	if (Ar.IsLoading())
+	{
+		int32 Value = 0;
+		Ar << Value;
+		FileVersion = (EVFCFileVersion)Value;
+	}
+	else
+	{
+		int32 Value = (int32)FileVersion;
+		Ar << Value;
+	}
+
+	return Ar;
+}
+
+static FArchive& operator<<(FArchive& Ar, FBlockFile& BlockFile)
+{
+	Ar << BlockFile.FileVersion;
+
+	if (Ar.IsLoading() && BlockFile.FileVersion != EVFCFileVersion::Current)
+	{
+		Ar.SetError();
+		BlockFile.Reset();
+		return Ar;
+	}
+
+	Ar << BlockFile.FileId;
+	Ar << BlockFile.BlockSize;
+	Ar << BlockFile.NumBlocks;
+	Ar << BlockFile.FreeRanges;
+	Ar << BlockFile.UsedRanges;
+	return Ar;
+}
+
+FArchive& operator<<(FArchive& Ar, FFileTable& FileTable)
+{
+	Ar << FileTable.FileVersion;
+
+	if (Ar.IsLoading() && FileTable.FileVersion != EVFCFileVersion::Current)
+	{
+		Ar.SetError();
+		FileTable.Empty();
+		return Ar;
+	}
+
+	Ar << FileTable.BlockFiles;
+	Ar << FileTable.FileMap;
+	Ar << FileTable.LastBlockFileId;
+	return Ar;
+}
diff --git a/Engine/Source/Runtime/VirtualFileCache/Private/VirtualFileCacheInternal.h b/Engine/Source/Runtime/VirtualFileCache/Private/VirtualFileCacheInternal.h
new file mode 100644
--- /dev/null
+++ b/Engine/Source/Runtime/VirtualFileCache/Private/VirtualFileCacheInternal.h
@@ -0,0 +1,374 @@
+// Copyright Epic Games, Inc. All Rights Reserved.
+
+#pragma once
+#include "VirtualFileCache.h"
+
+#include "CoreMinimal.h"
+#include "HAL/FileManager.h"
+#include "Logging/LogMacros.h"
+#include "Misc/Paths.h"
+#include "Misc/ScopeExit.h"
+#include "Async/MappedFileHandle.h"
+#include "Async/Future.h"
+
+// Fragmentation is defined as the amount of separation per file in the cache.
+// A value of 1 means that there is no fragmentation
+// A value of 2 means that the mean number of ranges per file is 2
+constexpr double VFC_ALLOWED_FRAGMENTATION = 4.0f;
+
+constexpr uint64 VFC_BLOCK_SIZE_MIN = 1;
+constexpr uint64 VFC_BLOCK_SIZE_MAX = 1 << 24;
+
+using FMappedRange = TUniquePtr<IMappedFileRegion>;
+
+enum class EVFCFileVersion
+{
+	Invalid,
+	Initial,
+
+	Count,
+	Current = Count - 1,
+};
+
+struct FBlockRange
+{
+	int32 StartIndex;
+	int32 NumBlocks;
+};
+
+struct FRangeId
+{
+	int32 FileId;
+	FBlockRange Range;
+};
+
+struct FDataReference
+{
+	TArray<FRangeId> Ranges;
+	uint32 TotalSize = 0;
+
+	mutable int64 LastReferencedUnixTime = 0; // Mutable since it needs to update when the data is read
+	void Touch();
+};
+
+struct FBlockFile
+{
+	// Serialized:
+	EVFCFileVersion FileVersion = EVFCFileVersion::Current;
+	int32 FileId = 0;
+	int32 BlockSize = 0;
+	int32 NumBlocks = 0;
+
+	TArray<FBlockRange> FreeRanges;
+	TArray<FBlockRange> UsedRanges;
+
+	// Not Serialized:
+	TUniquePtr<IFileHandle> WriteHandle;
+	TUniquePtr<IMappedFileHandle> MapHandle;
+	bool bWriteLocked = false;
+	TUniquePtr<FRWLock> FileHandleLock = MakeUnique<FRWLock>();
+
+	int64 TotalSize()
+	{
+		return static_cast<int64>(NumBlocks) * BlockSize;
+	}
+
+	void Reset()
+	{
+		FileId = 0;
+		BlockSize = 0;
+		NumBlocks = 0;
+		FreeRanges.Empty();
+		UsedRanges.Empty();
+		WriteHandle = nullptr;
+		MapHandle = nullptr;
+		bWriteLocked = false;
+		FileHandleLock = MakeUnique<FRWLock>();
+	}
+};
+
+using FFileMap = TMap<VFCKey, FDataReference>;
+struct FVirtualFileCache;
+
+struct IFileTableReader
+{
+	virtual ~IFileTableReader() = default;
+
+	virtual bool ValidateRanges() = 0;
+	virtual const FDataReference* FindDataReference(const VFCKey& Id) const = 0;
+	virtual TIoStatusOr<TArray<uint8>> ReadData(VFCKey Id, int64 ReadOffset = 0, int64 ReadSizeOrZero = 0) = 0;
+	virtual TIoStatusOr<uint64> GetSizeForChunk(const VFCKey& Id) const = 0;
+	virtual bool DoesChunkExist(const VFCKey& Id) const = 0;
+	virtual double CurrentFragmentation() const = 0;
+};
+
+struct FFileTable : IFileTableReader
+{
+	FFileTable(FVirtualFileCache* InParent)
+		: Parent(InParent)
+	{ }
+
+private:
+	FVirtualFileCache* Parent;
+	EVFCFileVersion FileVersion = EVFCFileVersion::Current;
+
+	/* Files on disk */
+	TArray<FBlockFile> BlockFiles;
+
+	/* Data stored in the chunks */
+	FFileMap FileMap;
+
+	int64 TotalSize = -1;
+	int64 UsedSize = -1;
+	int32 LastBlockFileId = 0;
+
+	int32 WriteLockCount = 0;
+
+public:
+	void Initialize(const FVirtualFileCacheSettings& Settings);
+	bool ReadTableFile();
+	int32 CreateBlockFile(int64 FileSize, int32 BlockSize);
+	bool DeleteBlockFile(int32 BlockFileId);
+	void WriteTableFile();
+	void Defragment();
+
+	void Empty();
+	bool CoalesceRanges();
+	void CalculateSizes();
+
+	FString GetBlockFilename(const FBlockFile& File) const;
+	FBlockFile* GetFileForRange(FRangeId Id);
+	FBlockFile* GetFileForId(int32 BlockFileId);
+	bool BlockFileExistsOnDisk(const FBlockFile& BlockFile) const;
+	IFileHandle* OpenBlockFileForWrite(FBlockFile& BlockFile);
+	bool DeleteBlockFile(FFileTable* FileTable, int32 BlockFileId);
+	bool RangeIsValid(FRangeId Block);
+	bool EraseData(VFCKey Id);
+
+	TIoStatusOr<FRangeId> AllocateSingleRange(FBlockFile& File, int64 MaximumSize);
+	TArray<FRangeId> AllocateBlocksForSize(uint64 Size);
+	int64 AllocationSize(FDataReference* DataRef);
+
+	void FreeBlock(FRangeId RangeId);
+	bool EnsureSizeFor(int64 RequiredBytes);
+	int64 EvictOne();
+	bool EvictAmount(int64 NumBytesToEvict);
+
+	IMappedFileHandle* MapBlockFile(FBlockFile& BlockFile);
+	FMappedRange MapFileRange(FRangeId Range);
+	bool ReadRange(FRangeId RangeId, void* Dest, int64 ReadSize);
+
+	FIoStatus WriteData(VFCKey Id, const uint8* Data, uint64 DataSize);
+
+	// IFileTableReader
+	virtual bool ValidateRanges() override;
+	virtual const FDataReference* FindDataReference(const VFCKey& Id) const override;
+	virtual TIoStatusOr<TArray<uint8>> ReadData(VFCKey Id, int64 ReadOffset = 0, int64 ReadSizeOrZero = 0) override;
+	virtual TIoStatusOr<uint64> GetSizeForChunk(const VFCKey& Id) const override;
+	virtual bool DoesChunkExist(const VFCKey& Id) const override;
+	virtual double CurrentFragmentation() const override;
+	// \IFileTableReader
+
+	friend FArchive& operator<<(FArchive& Ar, FFileTable& FileTable);
+};
+
+struct FBlockRangeSortStartIndex
+{
+	bool operator()(const FBlockRange& L, const FBlockRange& R) const
+	{
+		return L.StartIndex < R.StartIndex;
+	}
+};
+struct FBlockRangeSortSize
+{
+	bool operator()(const FBlockRange& L, const FBlockRange& R) const
+	{
+		return L.NumBlocks > R.NumBlocks;
+	}
+};
+
+inline FString GetVFCDirectory()
+{
+	return FPaths::ProjectPersistentDownloadDir() / "VFC";
+}
+FArchive& operator<<(FArchive& Ar, FFileTable& FileTable);
+
+enum class ERWOp : int8
+{
+	Invalid = -1,
+	Read,
+	Write,
+	Erase,
+};
+
+struct FRWOp
+{
+	ERWOp Op = ERWOp::Invalid;
+	VFCKey Target;
+
+	// Write
+	TArray<uint8> DataToWrite;
+
+	// Read
+	TOptional<TPromise<TArray<uint8>>> ReadResult;
+	int64 ReadOffset = 0;
+	int64 ReadSize = 0;
+};
+
+struct FVirtualFileCache;
+
+template <typename TTableType, FRWScopeLockType LOCKTYPE>
+struct FFileTableLockingReference
+{
+	static_assert(LOCKTYPE == SLT_ReadOnly || LOCKTYPE == SLT_Write);
+
+	FFileTableLockingReference(TTableType& InTable, FRWLock& InLock)
+		: FileTable(&InTable)
+		, Lock(&InLock)
+	{
+		if constexpr (LOCKTYPE == SLT_ReadOnly)
+		{
+			Lock->ReadLock();
+		}
+		else if constexpr (LOCKTYPE == SLT_Write)
+		{
+			Lock->WriteLock();
+		}
+	}
+
+	~FFileTableLockingReference()
+	{
+		if (Lock)
+		{
+			if constexpr (LOCKTYPE == SLT_ReadOnly)
+			{
+				Lock->ReadUnlock();
+			}
+			else if constexpr (LOCKTYPE == SLT_Write)
+			{
+				Lock->WriteUnlock();
+			}
+		}
+	}
+
+	FFileTableLockingReference(FFileTableLockingReference&& ToMove)
+		: FileTable(ToMove.FileTable)
+		, Lock(ToMove.Lock)
+	{
+		ToMove.Lock = nullptr;
+	}
+
+	FFileTableLockingReference(FFileTableLockingReference& ToMove) = delete;
+	FFileTableLockingReference operator=(FFileTableLockingReference& ToMove) = delete;
+	FFileTableLockingReference& operator=(FFileTableLockingReference&& ToMove) = delete;
+
+	TTableType* operator->()
+	{
+		return FileTable;
+	}
+
+	TTableType* Get()
+	{
+		return FileTable;
+	}
+
+private:
+	TTableType* FileTable;
+	FRWLock* Lock;
+};
+
+using FFileTableReader = FFileTableLockingReference<const IFileTableReader, SLT_ReadOnly>;
+using FFileTableMutator = FFileTableLockingReference<FFileTable, SLT_ReadOnly>;
+using FFileTableWriter = FFileTableLockingReference<FFileTable, SLT_Write>;
+
+struct FVirtualFileCacheThread : public FRunnable
+{
+	FVirtualFileCacheThread(FVirtualFileCache* InParent);
+	~FVirtualFileCacheThread();
+
+	/* FRunnable */
+	virtual bool Init() override;
+	virtual uint32 Run() override;
+	virtual void Stop() override;
+	virtual void Exit() override;
+	/* FRunnable */
+
+	// Public Interface
+	TFuture<TArray<uint8>> RequestRead(VFCKey Target, int64 ReadOffset = 0, int64 ReadSizeOrZero = 0);
+	void RequestWrite(VFCKey Target, TArrayView<const uint8> Data);
+	void RequestErase(VFCKey Target);
+	void Shutdown();
+
+	UE_NODISCARD FFileTableReader ReadFileTable() const;
+	UE_NODISCARD FFileTableMutator MutateFileTable();
+	UE_NODISCARD FFileTableWriter ModifyFileTable();
+
+public:
+	void DeleteUnexpectedCacheFiles(TSet<FString>& ExpectedFiles);
+	void FreeBlock(FRangeId RangeId);
+	FString GetTableFilename() const;
+
+	void EnqueueOrRunOp(TSharedPtr<FRWOp> Op);
+	TSharedPtr<FRWOp> GetNextOp();
+	void DoOneOp(FRWOp* Op);
+
+
+	void Touch(FDataReference& Id);
+
+	void EraseTableFile();
+
+
+public:
+	FVirtualFileCache* Parent;
+	FRunnableThread* Thread;
+	FEvent* Event;
+
+	FRWLock OperationQueueLock;
+	TArray<TSharedPtr<FRWOp>> OperationQueue;
+
+	FString BasePath;
+
+	std::atomic_bool bStopRequested;
+
+private:
+	mutable FRWLock FileTableLock;
+	FFileTable FileTableStorage;
+};
+
+struct FVirtualFileCache final : IVirtualFileCache
+{
+	FVirtualFileCache()
+		: Thread(this)
+	{}
+	virtual ~FVirtualFileCache()
+	{
+		Shutdown();
+	}
+	FVirtualFileCache(const FVirtualFileCache&) = delete;
+	FVirtualFileCache& operator=(const FVirtualFileCache&) = delete;
+
+public:
+	void Shutdown();
+
+	virtual void Initialize(const FVirtualFileCacheSettings& Settings) override;
+	virtual FIoStatus WriteData(VFCKey Id, const uint8* Data, uint64 DataSize) override;
+	virtual TFuture<TArray<uint8>> ReadData(VFCKey Id, int64 ReadOffset = 0, int64 ReadSizeOrZero = 0) override;
+	virtual bool DoesChunkExist(const VFCKey& Id) const override;
+	virtual void EraseData(VFCKey Id) override;
+	virtual TIoStatusOr<uint64> GetSizeForChunk(const VFCKey& Id) const override;
+	virtual double CurrentFragmentation() const override;
+	virtual void Defragment() override;
+
+	FString GetTableFilename() const { return Thread.GetTableFilename(); }
+
+public:
+	FChunkEvictedDelegate OnDataEvicted;
+	FVirtualFileCacheThread Thread;
+	FString BasePath;
+
+private:
+	FVirtualFileCacheSettings Settings;
+
+	friend struct FVirtualFileCacheThread;
+};
+
diff --git a/Engine/Source/Runtime/VirtualFileCache/Public/VirtualFileCache.h b/Engine/Source/Runtime/VirtualFileCache/Public/VirtualFileCache.h
new file mode 100644
--- /dev/null
+++ b/Engine/Source/Runtime/VirtualFileCache/Public/VirtualFileCache.h
@@ -0,0 +1,39 @@
+// Copyright Epic Games, Inc. All Rights Reserved.
+
+#pragma once
+
+#include "CoreMinimal.h"
+#include "IO/IoDispatcher.h"
+
+constexpr uint64 VFS_DEFAULT_BLOCK_SIZE = 8192ull;
+constexpr uint64 VFS_BLOCK_FILE_SIZE = 2 * 1024ull * 1024ull * 1024ull;
+
+using VFCKey = FSHAHash;
+
+DECLARE_MULTICAST_DELEGATE_OneParam(FChunkEvictedDelegate, VFCKey);
+
+struct FVirtualFileCacheSettings
+{
+	uint64 BlockSize = VFS_DEFAULT_BLOCK_SIZE;
+	uint64 BlockFileSize = VFS_BLOCK_FILE_SIZE;
+	uint64 NumBlockFiles = 1;
+	FString OverrideDefaultDirectory;
+	bool bMultiThreaded = true;
+};
+
+struct IVirtualFileCache
+{
+	virtual ~IVirtualFileCache() = default;
+
+	virtual void Initialize(const FVirtualFileCacheSettings& Context = {}) = 0;
+	virtual FIoStatus WriteData(VFCKey Id, const uint8* Data, uint64 DataSize) = 0;
+	virtual TFuture<TArray<uint8>> ReadData(VFCKey Id, int64 ReadOffset = 0, int64 ReadSizeOrZero = 0) = 0;
+	virtual bool DoesChunkExist(const VFCKey& Id) const = 0;
+	virtual void EraseData(VFCKey Id) = 0;
+	virtual TIoStatusOr<uint64> GetSizeForChunk(const VFCKey& Id) const = 0;
+	virtual double CurrentFragmentation() const = 0;
+	virtual void Defragment() = 0;
+
+	static VIRTUALFILECACHE_API TSharedRef<IVirtualFileCache> CreateVirtualFileCache();
+};
+
diff --git a/Engine/Source/Runtime/VirtualFileCache/VirtualFileCache.Build.cs b/Engine/Source/Runtime/VirtualFileCache/VirtualFileCache.Build.cs
new file mode 100644
--- /dev/null
+++ b/Engine/Source/Runtime/VirtualFileCache/VirtualFileCache.Build.cs
@@ -0,0 +1,16 @@
+// Copyright Epic Games, Inc. All Rights Reserved.
+
+using UnrealBuildTool;
+
+public class VirtualFileCache : ModuleRules
+{
+	public VirtualFileCache(ReadOnlyTargetRules Target) : base(Target)
+	{
+		PrivateDependencyModuleNames.AddRange(
+			new string[] {
+				"Core",
+				"CoreUObject",
+			}
+		);
+	}
+}
