Add header to fix non-unity compile error

diff --git a/Engine/Source/Runtime/VirtualFileCache/Private/VirtualFileCacheInternal.h b/Engine/Source/Runtime/VirtualFileCache/Private/VirtualFileCacheInternal.h
--- a/Engine/Source/Runtime/VirtualFileCache/Private/VirtualFileCacheInternal.h
+++ b/Engine/Source/Runtime/VirtualFileCache/Private/VirtualFileCacheInternal.h
@@ -8,6 +8,7 @@
 #include "Logging/LogMacros.h"
 #include "Misc/Paths.h"
 #include "Misc/ScopeExit.h"
+#include "Misc/ScopeRWLock.h"
 #include "Async/MappedFileHandle.h"
 #include "Async/Future.h"
 
