Fix template error in unity build

diff --git a/Engine/Source/Runtime/VirtualFileCache/Private/VirtualFileCacheInternal.h b/Engine/Source/Runtime/VirtualFileCache/Private/VirtualFileCacheInternal.h
--- a/Engine/Source/Runtime/VirtualFileCache/Private/VirtualFileCacheInternal.h
+++ b/Engine/Source/Runtime/VirtualFileCache/Private/VirtualFileCacheInternal.h
@@ -278,9 +278,9 @@ private:
 	FRWLock* Lock;
 };
 
-using FFileTableReader = FFileTableLockingReference<const IFileTableReader, SLT_ReadOnly>;
-using FFileTableMutator = FFileTableLockingReference<FFileTable, SLT_ReadOnly>;
-using FFileTableWriter = FFileTableLockingReference<FFileTable, SLT_Write>;
+using FFileTableReader = FFileTableLockingReference<const IFileTableReader, FRWScopeLockType::SLT_ReadOnly>;
+using FFileTableMutator = FFileTableLockingReference<FFileTable, FRWScopeLockType::SLT_ReadOnly>;
+using FFileTableWriter = FFileTableLockingReference<FFileTable, FRWScopeLockType::SLT_Write>;
 
 struct FVirtualFileCacheThread : public FRunnable
 {
