Unity fixes in VirtualFileCache

diff --git a/Engine/Source/Runtime/VirtualFileCache/Private/VirtualFileCache.cpp b/Engine/Source/Runtime/VirtualFileCache/Private/VirtualFileCache.cpp
--- a/Engine/Source/Runtime/VirtualFileCache/Private/VirtualFileCache.cpp
+++ b/Engine/Source/Runtime/VirtualFileCache/Private/VirtualFileCache.cpp
@@ -10,22 +10,15 @@
 #include "Misc/ScopeExit.h"
 #include "Async/MappedFileHandle.h"
 
-DECLARE_STATS_GROUP(TEXT("VFC"), STATGROUP_VFC, STATCAT_Advanced);
+DEFINE_STAT(STAT_FilesAdded);
+DEFINE_STAT(STAT_BytesAdded);
+DEFINE_STAT(STAT_FilesRemoved);
+DEFINE_STAT(STAT_BytesRemoved);
+DEFINE_STAT(STAT_FilesEvicted);
+DEFINE_STAT(STAT_BytesEvicted);
 
-DECLARE_DWORD_COUNTER_STAT(TEXT("Files Added"), STAT_FilesAdded, STATGROUP_VFC);
-DECLARE_DWORD_COUNTER_STAT(TEXT("Bytes Added"), STAT_BytesAdded, STATGROUP_VFC);
-DECLARE_DWORD_COUNTER_STAT(TEXT("Files Removed"), STAT_FilesRemoved, STATGROUP_VFC);
-DECLARE_DWORD_COUNTER_STAT(TEXT("Bytes Removed"), STAT_BytesRemoved, STATGROUP_VFC);
-DECLARE_DWORD_COUNTER_STAT(TEXT("Files Evicted"), STAT_FilesEvicted, STATGROUP_VFC);
-DECLARE_DWORD_COUNTER_STAT(TEXT("Bytes Evicted"), STAT_BytesEvicted, STATGROUP_VFC);
-
-DECLARE_LOG_CATEGORY_EXTERN(LogVFC, Log, All);
 DEFINE_LOG_CATEGORY(LogVFC);
 
-static const TCHAR VFC_CACHE_FILE_BASE_NAME[] = TEXT("vfc_");
-static const TCHAR VFC_CACHE_FILE_EXTENSION[] = TEXT("data");
-static const TCHAR VFC_META_FILE_NAME[] = TEXT("vfc.meta");
-
 TSharedRef<IVirtualFileCache> IVirtualFileCache::CreateVirtualFileCache()
 {
 	static TWeakPtr<FVirtualFileCache> GVFC;
diff --git a/Engine/Source/Runtime/VirtualFileCache/Private/VirtualFileCacheInternal.cpp b/Engine/Source/Runtime/VirtualFileCache/Private/VirtualFileCacheInternal.cpp
--- a/Engine/Source/Runtime/VirtualFileCache/Private/VirtualFileCacheInternal.cpp
+++ b/Engine/Source/Runtime/VirtualFileCache/Private/VirtualFileCacheInternal.cpp
@@ -9,6 +9,10 @@
 #include "Misc/ScopeExit.h"
 #include "Async/MappedFileHandle.h"
 
+static const TCHAR VFC_CACHE_FILE_BASE_NAME[] = TEXT("vfc_");
+static const TCHAR VFC_CACHE_FILE_EXTENSION[] = TEXT("data");
+static const TCHAR VFC_META_FILE_NAME[] = TEXT("vfc.meta");
+
 static int64 CurrentTimestamp()
 {
 	return FDateTime::UtcNow().ToUnixTimestamp();
diff --git a/Engine/Source/Runtime/VirtualFileCache/Private/VirtualFileCacheInternal.h b/Engine/Source/Runtime/VirtualFileCache/Private/VirtualFileCacheInternal.h
--- a/Engine/Source/Runtime/VirtualFileCache/Private/VirtualFileCacheInternal.h
+++ b/Engine/Source/Runtime/VirtualFileCache/Private/VirtualFileCacheInternal.h
@@ -12,6 +12,16 @@
 #include "Async/MappedFileHandle.h"
 #include "Async/Future.h"
 
+DECLARE_LOG_CATEGORY_EXTERN(LogVFC, Log, All);
+
+DECLARE_STATS_GROUP(TEXT("VFC"), STATGROUP_VFC, STATCAT_Advanced);
+DECLARE_DWORD_COUNTER_STAT_EXTERN(TEXT("Files Added"), STAT_FilesAdded, STATGROUP_VFC, );
+DECLARE_DWORD_COUNTER_STAT_EXTERN(TEXT("Bytes Added"), STAT_BytesAdded, STATGROUP_VFC, );
+DECLARE_DWORD_COUNTER_STAT_EXTERN(TEXT("Files Removed"), STAT_FilesRemoved, STATGROUP_VFC, );
+DECLARE_DWORD_COUNTER_STAT_EXTERN(TEXT("Bytes Removed"), STAT_BytesRemoved, STATGROUP_VFC, );
+DECLARE_DWORD_COUNTER_STAT_EXTERN(TEXT("Files Evicted"), STAT_FilesEvicted, STATGROUP_VFC, );
+DECLARE_DWORD_COUNTER_STAT_EXTERN(TEXT("Bytes Evicted"), STAT_BytesEvicted, STATGROUP_VFC, );
+
 // Fragmentation is defined as the amount of separation per file in the cache.
 // A value of 1 means that there is no fragmentation
 // A value of 2 means that the mean number of ranges per file is 2
