Fix missing header for Unity compile in VirtualFileCache

diff --git a/Engine/Source/Runtime/VirtualFileCache/Private/VirtualFileCacheInternal.h b/Engine/Source/Runtime/VirtualFileCache/Private/VirtualFileCacheInternal.h
--- a/Engine/Source/Runtime/VirtualFileCache/Private/VirtualFileCacheInternal.h
+++ b/Engine/Source/Runtime/VirtualFileCache/Private/VirtualFileCacheInternal.h
@@ -5,6 +5,7 @@
 
 #include "CoreMinimal.h"
 #include "HAL/FileManager.h"
+#include "HAL/Runnable.h"
 #include "Logging/LogMacros.h"
 #include "Misc/Paths.h"
 #include "Misc/ScopeExit.h"
