One more Unity header in VirtualFileCacheInternal.cpp

diff --git a/Engine/Source/Runtime/VirtualFileCache/Private/VirtualFileCacheInternal.cpp b/Engine/Source/Runtime/VirtualFileCache/Private/VirtualFileCacheInternal.cpp
--- a/Engine/Source/Runtime/VirtualFileCache/Private/VirtualFileCacheInternal.cpp
+++ b/Engine/Source/Runtime/VirtualFileCache/Private/VirtualFileCacheInternal.cpp
@@ -4,6 +4,7 @@
 
 #include "CoreMinimal.h"
 #include "HAL/FileManager.h"
+#include "HAL/Runnable.h"
 #include "Logging/LogMacros.h"
 #include "Misc/Paths.h"
 #include "Misc/ScopeExit.h"
