Adding an optional in-memory LRU cache to keep recently written data around temporarily since data is often immediately requested after being written to disk. The cache is disabled by default, but can be enabled by setting the RecentWriteLRUSize value in FVirtualFileCacheSettings before initializing the VFC.

diff --git a/Engine/Source/Runtime/VirtualFileCache/Private/VirtualFileCache.cpp b/Engine/Source/Runtime/VirtualFileCache/Private/VirtualFileCache.cpp
--- a/Engine/Source/Runtime/VirtualFileCache/Private/VirtualFileCache.cpp
+++ b/Engine/Source/Runtime/VirtualFileCache/Private/VirtualFileCache.cpp
@@ -48,6 +48,7 @@ void FVirtualFileCache::Initialize(const FVirtualFileCacheSettings& InSettings)
 
 	FFileTableWriter FileTable = Thread.ModifyFileTable();
 	FileTable->Initialize(Settings);
+	Thread.SetInMemoryCacheSize(Settings.RecentWriteLRUSize);
 }
 
 FIoStatus FVirtualFileCache::WriteData(VFCKey Id, const uint8* Data, uint64 DataSize)
diff --git a/Engine/Source/Runtime/VirtualFileCache/Private/VirtualFileCacheInternal.cpp b/Engine/Source/Runtime/VirtualFileCache/Private/VirtualFileCacheInternal.cpp
--- a/Engine/Source/Runtime/VirtualFileCache/Private/VirtualFileCacheInternal.cpp
+++ b/Engine/Source/Runtime/VirtualFileCache/Private/VirtualFileCacheInternal.cpp
@@ -35,7 +35,8 @@ FVirtualFileCacheThread::~FVirtualFileCacheThread()
 
 bool FVirtualFileCacheThread::Init()
 {
-	return Parent != nullptr;
+	check(Parent);
+	return true;
 }
 
 void FVirtualFileCacheThread::DoOneOp(FRWOp* Op)
@@ -44,23 +45,32 @@ void FVirtualFileCacheThread::DoOneOp(FRWOp* Op)
 	{
 	case ERWOp::Read:
 	{
-		FFileTableMutator FileTable = MutateFileTable();
-		auto Result = FileTable->ReadData(Op->Target, Op->ReadOffset, Op->ReadSize);
-		if (Result.IsOk())
+		// Check the memory cache first to see if this data was recently written
+		TSharedPtr<TArray<uint8>> CachedData = MemCache.ReadLockAndFindData(Op->Target);
+		if (CachedData.IsValid())
 		{
-			Op->ReadResult->SetValue(Result.ConsumeValueOrDie());
+			Op->ReadResult->SetValue(*CachedData);
 		}
 		else
 		{
-			TArray<uint8> Empty;
-			Op->ReadResult->SetValue(Empty);
+			FFileTableMutator FileTable = MutateFileTable();
+			auto Result = FileTable->ReadData(Op->Target, Op->ReadOffset, Op->ReadSize);
+			if (Result.IsOk())
+			{
+				Op->ReadResult->SetValue(Result.ConsumeValueOrDie());
+			}
+			else
+			{
+				TArray<uint8> Empty;
+				Op->ReadResult->SetValue(Empty);
+			}
 		}
 	} break;
 
 	case ERWOp::Write:
 	{
 		FFileTableWriter FileTable = ModifyFileTable();
-		FIoStatus Status = FileTable->WriteData(Op->Target, Op->DataToWrite.GetData(), Op->DataToWrite.Num());
+		FIoStatus Status = FileTable->WriteData(Op->Target, Op->DataToWrite->GetData(), Op->DataToWrite->Num());
 	} break;
 
 	case ERWOp::Erase:
@@ -132,6 +142,16 @@ void FVirtualFileCacheThread::EnqueueOrRunOp(TSharedPtr<FRWOp> Op)
 
 TFuture<TArray<uint8>> FVirtualFileCacheThread::RequestRead(VFCKey Target, int64 ReadOffset, int64 ReadSizeOrZero)
 {
+	// Check the memory cache first to see if this data was recently written. This will be checked
+	// again by the thread but should normally be caught here.
+	TSharedPtr<TArray<uint8>> CachedData = MemCache.ReadLockAndFindData(Target);
+	if (CachedData.IsValid())
+	{
+		TPromise<TArray<uint8>> ReadPromise;
+		ReadPromise.SetValue(*CachedData);
+		return ReadPromise.GetFuture();
+	}
+
 	TSharedPtr<FRWOp> Op = MakeShared<FRWOp>();
 	Op->Op = ERWOp::Read;
 	Op->Target = Target;
@@ -149,7 +169,10 @@ void FVirtualFileCacheThread::RequestWrite(VFCKey Target, TArrayView<const uint8
 	TSharedPtr<FRWOp> Op = MakeShared<FRWOp>();
 	Op->Op = ERWOp::Write;
 	Op->Target = Target;
-	Op->DataToWrite = MoveTemp(Data);
+	TSharedPtr<TArray<uint8>> SharedData = MakeShared<TArray<uint8>>(Data);
+	Op->DataToWrite = SharedData;
+
+	MemCache.Insert(Target, SharedData);
 
 	EnqueueOrRunOp(Op);
 	Event->Trigger();
@@ -255,6 +278,7 @@ void FFileTable::FreeBlock(FRangeId RangeId)
 	}
 }
 
+
 FIoStatus FFileTable::WriteData(VFCKey Id, const uint8* Data, uint64 DataSize)
 {
 	if (DoesChunkExist(Id))
@@ -1289,6 +1313,127 @@ void FVirtualFileCacheThread::EraseTableFile()
 	IPlatformFile::GetPlatformPhysical().DeleteFile(*TableFileName);
 }
 
+void FVirtualFileCacheThread::SetInMemoryCacheSize(int64 MaxSize)
+{
+	MemCache.SetMaxSize(MaxSize);
+}
+
+FLruCacheNode* FLruCache::Find(VFCKey Key)
+{
+	TUniquePtr<FLruCacheNode>* NodePtr = NodeMap.Find(Key);
+	if (NodePtr)
+	{
+		return NodePtr->Get();
+	}
+	return nullptr;
+}
+
+// Note that this function does not move the data to the front of the LRU, only writing data does that. This prevents
+// taking a write lock during read operations which would be required to modify the lru linked list.
+TSharedPtr<TArray<uint8>> FLruCache::ReadLockAndFindData(VFCKey Key)
+{
+	if (IsEnabled())
+	{
+		FRWScopeLock ScopeLock(Lock, SLT_ReadOnly);
+		FLruCacheNode* Node = Find(Key);
+		if (Node)
+		{
+			return Node->Data;
+		}
+	}
+	return nullptr;
+}
+
+void FLruCache::EvictToBelowMaxSize()
+{
+	while (CurrentSize > MaxSize)
+	{
+		EvictOne();
+	}
+}
+
+bool FLruCache::FreeSpaceFor(int64 SizeToAdd)
+{
+	if (SizeToAdd > MaxSize)
+	{
+		return false;
+	}
+
+
+	while (CurrentSize + SizeToAdd > MaxSize)
+	{
+		EvictOne();
+	}
+	return true;
+}
+
+void FLruCache::EvictOne()
+{
+	FLruCacheNode* Node = LruList.GetHead();
+	if (Node)
+	{
+		CurrentSize -= Node->RecordedSize;
+		LruList.Remove(Node);
+		int32 NumRemoved = NodeMap.Remove(Node->Key);
+		check(NumRemoved == 1);
+	}
+}
+
+void FLruCache::Insert(VFCKey Key, TSharedPtr<TArray<uint8>> Data)
+{
+	if (!IsEnabled())
+	{
+		return;
+	}
+
+	FRWScopeLock ScopeLock(Lock, SLT_Write);
+	FLruCacheNode* Existing = Find(Key);
+	if (Existing)
+	{
+		CurrentSize -= Existing->Data->Num();
+		Existing->Data = MoveTemp(Data);
+		CurrentSize += Existing->Data->Num();
+		LruList.Remove(Existing);
+		LruList.AddHead(Existing);
+	}
+	else if (FreeSpaceFor(Data->Num()))
+	{
+		TUniquePtr<FLruCacheNode> NewNode = MakeUnique<FLruCacheNode>();
+		NewNode->Key = Key;
+		NewNode->Data = MoveTemp(Data);
+		NewNode->RecordedSize = NewNode->Data->Num();
+		CurrentSize += NewNode->RecordedSize;
+		LruList.AddHead(NewNode.Get());
+		NodeMap.Add(Key, MoveTemp(NewNode));
+	}
+	else
+	{
+		UE_LOG(LogVFC, Verbose, TEXT("Data too large to fit in LRU cache (Requested: %zu, MaxSize: %zu"), (size_t)Data->Num(), (size_t)MaxSize);
+	}
+}
+
+void FLruCache::Remove(VFCKey Key)
+{
+	TUniquePtr<FLruCacheNode>* Existing = NodeMap.Find(Key);
+	if (Existing && Existing->IsValid())
+	{
+		FLruCacheNode* Ptr = Existing->Get();
+		CurrentSize -= Ptr->RecordedSize;
+		LruList.Remove(Ptr);
+		NodeMap.Remove(Key);
+	}
+}
+
+bool FLruCache::IsEnabled()
+{
+	return MaxSize > 0;
+}
+
+void FLruCache::SetMaxSize(int64 NewMaxSize)
+{
+	MaxSize = NewMaxSize;
+	EvictToBelowMaxSize();
+}
 
 static FArchive& operator<<(FArchive& Ar, FBlockRange& Range)
 {
diff --git a/Engine/Source/Runtime/VirtualFileCache/Private/VirtualFileCacheInternal.h b/Engine/Source/Runtime/VirtualFileCache/Private/VirtualFileCacheInternal.h
--- a/Engine/Source/Runtime/VirtualFileCache/Private/VirtualFileCacheInternal.h
+++ b/Engine/Source/Runtime/VirtualFileCache/Private/VirtualFileCacheInternal.h
@@ -12,6 +12,7 @@
 #include "Misc/ScopeRWLock.h"
 #include "Async/MappedFileHandle.h"
 #include "Async/Future.h"
+#include "Containers/IntrusiveDoubleLinkedList.h"
 
 DECLARE_LOG_CATEGORY_EXTERN(LogVFC, Log, All);
 
@@ -219,7 +220,7 @@ struct FRWOp
 	VFCKey Target;
 
 	// Write
-	TArray<uint8> DataToWrite;
+	TSharedPtr<TArray<uint8>> DataToWrite;
 
 	// Read
 	TOptional<TPromise<TArray<uint8>>> ReadResult;
@@ -227,6 +228,45 @@ struct FRWOp
 	int64 ReadSize = 0;
 };
 
+
+struct FLruCacheNode : public TIntrusiveDoubleLinkedListNode<FLruCacheNode>
+{
+	// For reverse lookups
+	VFCKey Key;
+	TSharedPtr<TArray<uint8>> Data;
+	// Store the size here to ensure that the CurrentSize of the LruCache can rely on the Data size value changing.
+	uint64 RecordedSize;
+};
+
+
+class FLruCache
+{
+	void EvictToBelowMaxSize();
+	void EvictOne();
+	bool FreeSpaceFor(int64 SizeToAdd);
+
+public:
+	~FLruCache()
+	{
+		LruList.Reset();
+	}
+	FLruCacheNode* Find(VFCKey Key);
+	TSharedPtr<TArray<uint8>> ReadLockAndFindData(VFCKey Key);
+	void Insert(VFCKey Key, TSharedPtr<TArray<uint8>> Data);
+	void Remove(VFCKey Key);
+	bool IsEnabled();
+	void SetMaxSize(int64 NewMaxSize);
+
+public:
+	FRWLock Lock;
+
+private:
+	TMap<VFCKey, TUniquePtr<FLruCacheNode>> NodeMap;
+	TIntrusiveDoubleLinkedList<FLruCacheNode> LruList;
+	int64 CurrentSize = 0;
+	int64 MaxSize = 0;
+};
+
 struct FVirtualFileCache;
 
 template <typename TTableType, FRWScopeLockType LOCKTYPE>
@@ -324,17 +364,18 @@ public:
 	TSharedPtr<FRWOp> GetNextOp();
 	void DoOneOp(FRWOp* Op);
 
-
 	void Touch(FDataReference& Id);
-
 	void EraseTableFile();
 
+	void SetInMemoryCacheSize(int64 MaxSize);
 
 public:
 	FVirtualFileCache* Parent;
 	FRunnableThread* Thread;
 	FEvent* Event;
 
+	FLruCache MemCache;
+
 	FRWLock OperationQueueLock;
 	TArray<TSharedPtr<FRWOp>> OperationQueue;
 
diff --git a/Engine/Source/Runtime/VirtualFileCache/Public/VirtualFileCache.h b/Engine/Source/Runtime/VirtualFileCache/Public/VirtualFileCache.h
--- a/Engine/Source/Runtime/VirtualFileCache/Public/VirtualFileCache.h
+++ b/Engine/Source/Runtime/VirtualFileCache/Public/VirtualFileCache.h
@@ -19,6 +19,9 @@ struct FVirtualFileCacheSettings
 	uint64 NumBlockFiles = 1;
 	FString OverrideDefaultDirectory;
 	bool bMultiThreaded = true;
+
+	// Memory Cache behavior
+	uint64 RecentWriteLRUSize = 0;
 };
 
 struct IVirtualFileCache
