VirtualFileCache was evicting the most recently used entry in the memory cache and not checking the memory cache when looking up a chunk

diff --git a/Engine/Source/Runtime/VirtualFileCache/Private/VirtualFileCache.cpp b/Engine/Source/Runtime/VirtualFileCache/Private/VirtualFileCache.cpp
--- a/Engine/Source/Runtime/VirtualFileCache/Private/VirtualFileCache.cpp
+++ b/Engine/Source/Runtime/VirtualFileCache/Private/VirtualFileCache.cpp
@@ -19,16 +19,16 @@ DEFINE_STAT(STAT_BytesEvicted);
 
 DEFINE_LOG_CATEGORY(LogVFC);
 
-TSharedRef<IVirtualFileCache> IVirtualFileCache::CreateVirtualFileCache()
+TSharedRef<IVirtualFileCache, ESPMode::ThreadSafe> IVirtualFileCache::CreateVirtualFileCache()
 {
-	static TWeakPtr<FVirtualFileCache> GVFC;
-	TSharedPtr<FVirtualFileCache> SharedVFC;
-	if (!GVFC.IsValid())
+	static TWeakPtr<FVirtualFileCache, ESPMode::ThreadSafe> GVFC;
+	TSharedPtr<FVirtualFileCache, ESPMode::ThreadSafe> SharedVFC = GVFC.Pin();
+	if (!SharedVFC.IsValid())
 	{
 		SharedVFC = MakeShared<FVirtualFileCache>();
 		GVFC = SharedVFC;
 	}
-	return GVFC.IsValid() ? GVFC.Pin().ToSharedRef() : SharedVFC.ToSharedRef();
+	return SharedVFC.ToSharedRef();
 }
 
 void FVirtualFileCache::Shutdown()
@@ -53,6 +53,14 @@ void FVirtualFileCache::Initialize(const FVirtualFileCacheSettings& InSettings)
 
 FIoStatus FVirtualFileCache::WriteData(VFCKey Id, const uint8* Data, uint64 DataSize)
 {
+	{
+		FFileTableReader FileTable = Thread.ReadFileTable();
+		if (FileTable->DoesChunkExist(Id))
+		{
+			return FIoStatus(EIoErrorCode::Ok);
+		}
+	}
+
 	Thread.RequestWrite(Id, MakeArrayView(Data, DataSize));
 	return FIoStatus(EIoErrorCode::Ok);
 }
@@ -64,6 +72,12 @@ TFuture<TArray<uint8>> FVirtualFileCache::ReadData(VFCKey Id, int64 ReadOffset,
 
 bool FVirtualFileCache::DoesChunkExist(const VFCKey& Id) const
 {
+	TSharedPtr<TArray<uint8>> CachedData = Thread.MemCache.ReadLockAndFindData(Id);
+	if (CachedData.IsValid())
+	{
+		return true;
+	}
+
 	FFileTableReader FileTable = Thread.ReadFileTable();
 	return FileTable->DoesChunkExist(Id);
 }
diff --git a/Engine/Source/Runtime/VirtualFileCache/Private/VirtualFileCacheInternal.cpp b/Engine/Source/Runtime/VirtualFileCache/Private/VirtualFileCacheInternal.cpp
--- a/Engine/Source/Runtime/VirtualFileCache/Private/VirtualFileCacheInternal.cpp
+++ b/Engine/Source/Runtime/VirtualFileCache/Private/VirtualFileCacheInternal.cpp
@@ -283,6 +283,7 @@ FIoStatus FFileTable::WriteData(VFCKey Id, const uint8* Data, uint64 DataSize)
 {
 	if (DoesChunkExist(Id))
 	{
+		UE_LOG(LogVFC, Verbose, TEXT("Overwriting hash %s of %llu bytes"), *Id.ToString(), DataSize);
 		EraseData(Id);
 	}
 
@@ -336,8 +337,6 @@ FIoStatus FFileTable::WriteData(VFCKey Id, const uint8* Data, uint64 DataSize)
 			//Handle->Flush();
 		}
 
-		UE_LOG(LogVFC, Verbose, TEXT("Wrote %llu bytes\n"), DataSize);
-
 		DataRef.TotalSize = DataSize;
 		DataRef.Touch();
 		FileMap.Add(Id, DataRef);
@@ -348,6 +347,8 @@ FIoStatus FFileTable::WriteData(VFCKey Id, const uint8* Data, uint64 DataSize)
 		UsedSize += TotalAllocatedSize;
 		WriteTableFile();
 
+		UE_LOG(LogVFC, Verbose, TEXT("Wrote hash %s of %llu bytes, total size is %lld out of %lld"), *Id.ToString(), DataSize, UsedSize, TotalSize);
+
 		INC_DWORD_STAT(STAT_FilesAdded);
 		INC_DWORD_STAT_BY(STAT_BytesAdded, DataSize);
 	}
@@ -534,7 +535,6 @@ int64 FFileTable::EvictOne()
 	{
 		if (MapValue.LastReferencedUnixTime >= 0 && MapValue.LastReferencedUnixTime < LeastRecentTime)
 		{
-			UE_LOG(LogVFC, Verbose, TEXT("Evicting Data %u with size of %u\n"), GetTypeHash(MapKey), MapValue.TotalSize);
 			LeastRecentlyUsed = MapKey;
 			LeastRecentTime = MapValue.LastReferencedUnixTime;
 		}
@@ -543,6 +543,8 @@ int64 FFileTable::EvictOne()
 	FDataReference* ToRemove = FileMap.Find(LeastRecentlyUsed);
 	if (ToRemove)
 	{
+		UE_LOG(LogVFC, Verbose, TEXT("Evicting Data %s with size of %u"), *LeastRecentlyUsed.ToString(), ToRemove->TotalSize);
+
 		RemovedSize = AllocationSize(ToRemove);
 		Parent->OnDataEvicted.Broadcast(LeastRecentlyUsed);
 		EraseData(LeastRecentlyUsed);
@@ -1328,14 +1330,24 @@ FLruCacheNode* FLruCache::Find(VFCKey Key)
 	return nullptr;
 }
 
+const FLruCacheNode* FLruCache::Find(VFCKey Key) const
+{
+	const TUniquePtr<FLruCacheNode>* NodePtr = NodeMap.Find(Key);
+	if (NodePtr)
+	{
+		return NodePtr->Get();
+	}
+	return nullptr;
+}
+
 // Note that this function does not move the data to the front of the LRU, only writing data does that. This prevents
 // taking a write lock during read operations which would be required to modify the lru linked list.
-TSharedPtr<TArray<uint8>> FLruCache::ReadLockAndFindData(VFCKey Key)
+TSharedPtr<TArray<uint8>> FLruCache::ReadLockAndFindData(VFCKey Key) const
 {
 	if (IsEnabled())
 	{
 		FRWScopeLock ScopeLock(Lock, SLT_ReadOnly);
-		FLruCacheNode* Node = Find(Key);
+		const FLruCacheNode* Node = Find(Key);
 		if (Node)
 		{
 			return Node->Data;
@@ -1369,9 +1381,10 @@ bool FLruCache::FreeSpaceFor(int64 SizeToAdd)
 
 void FLruCache::EvictOne()
 {
-	FLruCacheNode* Node = LruList.GetHead();
+	FLruCacheNode* Node = LruList.GetTail();
 	if (Node)
 	{
+		UE_LOG(LogVFC, Verbose, TEXT("Evicting from MemCache hash %s"), *Node->Key.ToString());
 		CurrentSize -= Node->RecordedSize;
 		LruList.Remove(Node);
 		int32 NumRemoved = NodeMap.Remove(Node->Key);
@@ -1390,6 +1403,7 @@ void FLruCache::Insert(VFCKey Key, TSharedPtr<TArray<uint8>> Data)
 	FLruCacheNode* Existing = Find(Key);
 	if (Existing)
 	{
+		UE_LOG(LogVFC, Verbose, TEXT("Advancing to head of MemCache hash %s"), *Key.ToString());
 		CurrentSize -= Existing->Data->Num();
 		Existing->Data = MoveTemp(Data);
 		CurrentSize += Existing->Data->Num();
@@ -1398,6 +1412,7 @@ void FLruCache::Insert(VFCKey Key, TSharedPtr<TArray<uint8>> Data)
 	}
 	else if (FreeSpaceFor(Data->Num()))
 	{
+		UE_LOG(LogVFC, Verbose, TEXT("Inserting into MemCache hash %s"), *Key.ToString());
 		TUniquePtr<FLruCacheNode> NewNode = MakeUnique<FLruCacheNode>();
 		NewNode->Key = Key;
 		NewNode->Data = MoveTemp(Data);
@@ -1424,7 +1439,7 @@ void FLruCache::Remove(VFCKey Key)
 	}
 }
 
-bool FLruCache::IsEnabled()
+bool FLruCache::IsEnabled() const
 {
 	return MaxSize > 0;
 }
diff --git a/Engine/Source/Runtime/VirtualFileCache/Private/VirtualFileCacheInternal.h b/Engine/Source/Runtime/VirtualFileCache/Private/VirtualFileCacheInternal.h
--- a/Engine/Source/Runtime/VirtualFileCache/Private/VirtualFileCacheInternal.h
+++ b/Engine/Source/Runtime/VirtualFileCache/Private/VirtualFileCacheInternal.h
@@ -251,14 +251,15 @@ public:
 		LruList.Reset();
 	}
 	FLruCacheNode* Find(VFCKey Key);
-	TSharedPtr<TArray<uint8>> ReadLockAndFindData(VFCKey Key);
+	const FLruCacheNode* Find(VFCKey Key) const;
+	TSharedPtr<TArray<uint8>> ReadLockAndFindData(VFCKey Key) const;
 	void Insert(VFCKey Key, TSharedPtr<TArray<uint8>> Data);
 	void Remove(VFCKey Key);
-	bool IsEnabled();
+	bool IsEnabled() const;
 	void SetMaxSize(int64 NewMaxSize);
 
 public:
-	FRWLock Lock;
+	mutable FRWLock Lock;
 
 private:
 	TMap<VFCKey, TUniquePtr<FLruCacheNode>> NodeMap;
diff --git a/Engine/Source/Runtime/VirtualFileCache/Public/VirtualFileCache.h b/Engine/Source/Runtime/VirtualFileCache/Public/VirtualFileCache.h
--- a/Engine/Source/Runtime/VirtualFileCache/Public/VirtualFileCache.h
+++ b/Engine/Source/Runtime/VirtualFileCache/Public/VirtualFileCache.h
@@ -37,6 +37,6 @@ struct IVirtualFileCache
 	virtual double CurrentFragmentation() const = 0;
 	virtual void Defragment() = 0;
 
-	static VIRTUALFILECACHE_API TSharedRef<IVirtualFileCache> CreateVirtualFileCache();
+	static VIRTUALFILECACHE_API TSharedRef<IVirtualFileCache, ESPMode::ThreadSafe> CreateVirtualFileCache();
 };
 
