Fixes a problem where valid cache files where deleted but the table remained intact.

diff --git a/Engine/Source/Runtime/VirtualFileCache/Private/VirtualFileCacheInternal.cpp b/Engine/Source/Runtime/VirtualFileCache/Private/VirtualFileCacheInternal.cpp
--- a/Engine/Source/Runtime/VirtualFileCache/Private/VirtualFileCacheInternal.cpp
+++ b/Engine/Source/Runtime/VirtualFileCache/Private/VirtualFileCacheInternal.cpp
@@ -948,7 +948,13 @@ void DeleteUnexpectedCacheFiles(const FString& BasePath, TSet<FString>& Expected
 				if (Extension == VFC_CACHE_FILE_EXTENSION)
 				{
 					FString Filename = FPaths::GetCleanFilename(FullPath);
-					if (!ExpectedFiles.Contains(Filename))
+					TArray<FString> ExpectedArray = ExpectedFiles.Array();
+					bool ContainsFile = false;
+					for (int32 i = 0; i < ExpectedArray.Num(); ++i)
+					{
+						ContainsFile = ContainsFile || ExpectedArray[i].Contains(Filename);
+					}
+					if (!ContainsFile)
 					{
 						IPlatformFile& PlatformFile = IPlatformFile::GetPlatformPhysical();
 						PlatformFile.DeleteFile(*FullPath);
