Virtual file cache changes
- Process more changes when the queue has items in it.
- Fix defragmentation issue.

diff --git a/Engine/Source/Runtime/VirtualFileCache/Private/VirtualFileCacheInternal.cpp b/Engine/Source/Runtime/VirtualFileCache/Private/VirtualFileCacheInternal.cpp
--- a/Engine/Source/Runtime/VirtualFileCache/Private/VirtualFileCacheInternal.cpp
+++ b/Engine/Source/Runtime/VirtualFileCache/Private/VirtualFileCacheInternal.cpp
@@ -85,7 +85,6 @@ uint32 FVirtualFileCacheThread::Run()
 {
 	while (!bStopRequested)
 	{
-		Event->Wait(100);
 		TSharedPtr<FRWOp> Op = GetNextOp();
 		if (bStopRequested)
 		{
@@ -94,6 +93,7 @@ uint32 FVirtualFileCacheThread::Run()
 		}
 		if (!Op.IsValid())
 		{
+			Event->Wait(100);
 			continue;
 		}
 		DoOneOp(Op.Get());
@@ -780,8 +780,6 @@ void FFileTable::Defragment()
 	}
 	CalculateSizes();
 
-	int32 TESTCounter = 0;
-
 	if (bSuccess)
 	{
 		TArray<VFCKey> DataToMove;
@@ -803,7 +801,7 @@ void FFileTable::Defragment()
 				TArray<uint8> Data = ReadResult.ConsumeValueOrDie();
 				FIoStatus WriteResult = WriteData(Id, Data.GetData(), Data.Num());
 
-				if (!WriteResult.IsOk() || TESTCounter++ == 5)
+				if (!WriteResult.IsOk())
 				{
 					bSuccess = false;
 					break;
