Fix for appending block files not accounting for existing block files
#localization none
#tests Compilation, golden path, loading existing cache and reviewing migration, Additional QA coverage

diff --git a/Engine/Source/Runtime/VirtualFileCache/Private/VirtualFileCacheInternal.cpp b/Engine/Source/Runtime/VirtualFileCache/Private/VirtualFileCacheInternal.cpp
--- a/Engine/Source/Runtime/VirtualFileCache/Private/VirtualFileCacheInternal.cpp
+++ b/Engine/Source/Runtime/VirtualFileCache/Private/VirtualFileCacheInternal.cpp
@@ -1019,7 +1019,7 @@ void FFileTable::Initialize(const FVirtualFileCacheSettings& Settings)
 		int64 CacheFileSize = Settings.BlockFileSize;
 		BlockSize = FMath::Min(BlockSize, CacheFileSize);
 
-		for (int32 i = 0; i < Settings.NumBlockFiles; ++i)
+		for (int32 i = BlockFiles.Num(); i < Settings.NumBlockFiles; ++i)
 		{
 			if (CreateBlockFile(CacheFileSize, BlockSize) == 0)
 			{
